[NO NON-PII COMMIT MESSAGE]

diff --git a/Engine/Plugins/Experimental/PlanarCutPlugin/Source/PlanarCut/Private/PlanarCut.cpp b/Engine/Plugins/Experimental/PlanarCutPlugin/Source/PlanarCut/Private/PlanarCut.cpp
--- a/Engine/Plugins/Experimental/PlanarCutPlugin/Source/PlanarCut/Private/PlanarCut.cpp
+++ b/Engine/Plugins/Experimental/PlanarCutPlugin/Source/PlanarCut/Private/PlanarCut.cpp
@@ -1,5 +1,6 @@
 // Copyright 1998-2019 Epic Games, Inc. All Rights Reserved.
 #include "PlanarCut.h"
+#include "PlanarCutPlugin.h"
 
 #include "Async/ParallelFor.h"
 #include "Spatial/FastWinding.h"
@@ -2110,7 +2111,7 @@ int32 CutMultipleWithPlanarCells(
 		if (Source.Children[ParentTransformIndex].Num())
 		{
 			// don't fracture an already-fractured geometry
-			ensureMsgf(false, TEXT("Skipping cut of a non-leaf geometry, as this would would create intersecting / duplicate geometry"));
+			UE_LOG(LogPlanarCut, Warning, TEXT("Skipping cut of a non-leaf geometry, as this would would create intersecting / duplicate geometry"));
 			continue;
 		}
 		int32 TriangleStart = Source.FaceStart[GeometryIdx];
diff --git a/Engine/Plugins/Experimental/PlanarCutPlugin/Source/PlanarCut/Private/PlanarCutPlugin.cpp b/Engine/Plugins/Experimental/PlanarCutPlugin/Source/PlanarCut/Private/PlanarCutPlugin.cpp
--- a/Engine/Plugins/Experimental/PlanarCutPlugin/Source/PlanarCut/Private/PlanarCutPlugin.cpp
+++ b/Engine/Plugins/Experimental/PlanarCutPlugin/Source/PlanarCut/Private/PlanarCutPlugin.cpp
@@ -6,6 +6,7 @@
 #include "Modules/ModuleManager.h"
 
 
+DEFINE_LOG_CATEGORY(LogPlanarCut);
 
 
 class FPlanarCutPlugin : public IPlanarCutPlugin
diff --git a/Engine/Plugins/Experimental/PlanarCutPlugin/Source/PlanarCut/Public/PlanarCutPlugin.h b/Engine/Plugins/Experimental/PlanarCutPlugin/Source/PlanarCut/Public/PlanarCutPlugin.h
--- a/Engine/Plugins/Experimental/PlanarCutPlugin/Source/PlanarCut/Public/PlanarCutPlugin.h
+++ b/Engine/Plugins/Experimental/PlanarCutPlugin/Source/PlanarCut/Public/PlanarCutPlugin.h
@@ -6,6 +6,7 @@
 #include "Modules/ModuleInterface.h"
 #include "Modules/ModuleManager.h"
 
+DECLARE_LOG_CATEGORY_EXTERN(LogPlanarCut, Log, All);
 
 /**
  * The public interface to this module
