Change planar cut api to take transform indices to cut instead of geometry indices, to make more convenient/natural/expected api

#rnx

diff --git a/Engine/Plugins/Experimental/PlanarCutPlugin/Source/PlanarCut/Private/PlanarCut.cpp b/Engine/Plugins/Experimental/PlanarCutPlugin/Source/PlanarCut/Private/PlanarCut.cpp
--- a/Engine/Plugins/Experimental/PlanarCutPlugin/Source/PlanarCut/Private/PlanarCut.cpp
+++ b/Engine/Plugins/Experimental/PlanarCutPlugin/Source/PlanarCut/Private/PlanarCut.cpp
@@ -2037,7 +2037,7 @@ void TransformPlanes(const FTransform& Transform, const FPlanarCells& Ref, TArra
 int32 CutWithPlanarCells(
 	FPlanarCells& Cells,
 	FGeometryCollection& Source,
-	int32 GeometryIdx,
+	int32 TransformIdx,
 	const TOptional<FTransform>& TransformCells,
 	bool bIncludeOutsideCellInOutput,
 	float CheckDistanceAcrossOutsideCellForProximity,
@@ -2045,15 +2045,15 @@ int32 CutWithPlanarCells(
 	TFunction<void(const FGeometryCollection&, int32, const FGeometryCollection&, int32, float, int32, FGeometryCollection&)> VertexInterpolate
 )
 {
-	TArray<int32> GeometryIndices { GeometryIdx };
-	return CutMultipleWithPlanarCells(Cells, Source, GeometryIndices, TransformCells, bIncludeOutsideCellInOutput, CheckDistanceAcrossOutsideCellForProximity, bSetDefaultInternalMaterialsFromCollection, VertexInterpolate);
+	TArray<int32> TransformIndices { TransformIdx };
+	return CutMultipleWithPlanarCells(Cells, Source, TransformIndices, TransformCells, bIncludeOutsideCellInOutput, CheckDistanceAcrossOutsideCellForProximity, bSetDefaultInternalMaterialsFromCollection, VertexInterpolate);
 }
 
 // Cut multiple Geometry groups inside a GeometryCollection with PlanarCells, and add each cut cell back to the GeometryCollection as a new child of their source Geometry
 int32 CutMultipleWithPlanarCells(
 	FPlanarCells &Cells,
 	FGeometryCollection& Source,
-	const TArrayView<const int32>& GeometryIndices,
+	const TArrayView<const int32>& TransformIndices,
 	const TOptional<FTransform>& TransformCells,
 	bool bIncludeOutsideCellInOutput,
 	float CheckDistanceAcrossOutsideCellForProximity,
@@ -2086,8 +2086,9 @@ int32 CutMultipleWithPlanarCells(
 	static const FText SlowTaskText = NSLOCTEXT("CutMultipleWithPlanarCells", "CutMultipleWithPlanarCellsText", "Cutting geometry collection...");
 
 	int32 TotalFacesToProcess = 0;
-	for (int32 GeometryIdx : GeometryIndices)
+	for (int32 TransformIdx : TransformIndices)
 	{
+		int32 GeometryIdx = Source.TransformToGeometryIndex[TransformIdx];
 		TotalFacesToProcess += Source.FaceCount[GeometryIdx];
 	}
 
@@ -2104,10 +2105,10 @@ int32 CutMultipleWithPlanarCells(
 #endif
 
 	// TODO: Cannot yet make this a parallel for because the Voronoi pt->cell function is not thread-safe, but after that is fixed, consider threading this?
-	for (int32 GeometryIdx : GeometryIndices)
+	for (int32 ParentTransformIndex : TransformIndices)
 	{
+		int32 GeometryIdx = Source.TransformToGeometryIndex[ParentTransformIndex];
 		EnterProgressFrame(Source.FaceCount[GeometryIdx]);
-		int32 ParentTransformIndex = Source.TransformIndex[GeometryIdx];
 		if (Source.Children[ParentTransformIndex].Num())
 		{
 			// don't fracture an already-fractured geometry
@@ -2162,7 +2163,7 @@ int32 CutMultipleWithMultiplePlanes(
 	const TArrayView<const FPlane>& Planes,
 	FInternalSurfaceMaterials& InternalSurfaceMaterials,
 	FGeometryCollection& Collection,
-	const TArrayView<const int32>& GeometryIndices,
+	const TArrayView<const int32>& TransformIndices,
 	const TOptional<FTransform>& TransformCells,
 	bool bFlattenToSingleLayer,
 	bool bSetDefaultInternalMaterialsFromCollection,
@@ -2184,10 +2185,11 @@ int32 CutMultipleWithMultiplePlanes(
 	}
 
 	TArray<int32> TransformsToDelete;
-	NeedsCut.SetNum(GeometryIndices.Num());
+	NeedsCut.SetNum(TransformIndices.Num());
 	for (int32 Idx = 0; Idx < NeedsCut.Num(); Idx++)
 	{
-		NeedsCut[Idx] = GeometryIndices[Idx];
+		// add geometry idx corresponding to each input transform
+		NeedsCut[Idx] = Collection.TransformToGeometryIndex[TransformIndices[Idx]];
 	}
 
 	if (!Collection.HasAttribute("Proximity", FGeometryCollection::GeometryGroup))
diff --git a/Engine/Plugins/Experimental/PlanarCutPlugin/Source/PlanarCut/Public/PlanarCut.h b/Engine/Plugins/Experimental/PlanarCutPlugin/Source/PlanarCut/Public/PlanarCut.h
--- a/Engine/Plugins/Experimental/PlanarCutPlugin/Source/PlanarCut/Public/PlanarCut.h
+++ b/Engine/Plugins/Experimental/PlanarCutPlugin/Source/PlanarCut/Public/PlanarCut.h
@@ -147,7 +147,7 @@ void PLANARCUT_API ComputeTriangleNormals(const TArrayView<const FVector> Vertic
  * 
  * @param Cells				Defines the cutting planes and division of space
  * @param Collection		The collection to be cut
- * @param GeometryIdx		Which geometry group inside the collection to cut
+ * @param TransformIdx		Which transform inside the collection to cut
  * @param TransformCells	Optional transform of the planar cut; if unset, defaults to Identity
  * @param bIncludeOutsideCellInOutput	If true, geometry that was not inside any of the cells (e.g. was outside of the bounds of all cutting geometry) will still be included in the output; if false, it will be discarded.
  * @param CheckDistanceAcrossOutsideCellForProximity	If > 0, when a plane is neighboring the "outside" cell, instead of setting proximity to the outside cell, the algo will sample a point this far outside the cell in the normal direction of the plane to see if there is actually a non-outside cell there.  (Useful for bricks w/out mortar)
@@ -157,7 +157,7 @@ void PLANARCUT_API ComputeTriangleNormals(const TArrayView<const FVector> Vertic
 int32 PLANARCUT_API CutWithPlanarCells(
 	FPlanarCells &Cells,
 	FGeometryCollection& Collection,
-	int32 GeometryIdx,
+	int32 TransformIdx,
 	const TOptional<FTransform>& TransformCells = TOptional<FTransform>(),
 	bool bIncludeOutsideCellInOutput = true,
 	float CheckDistanceAcrossOutsideCellForProximity = 0,
@@ -169,8 +169,8 @@ int32 PLANARCUT_API CutWithPlanarCells(
  * Cut multiple Geometry groups inside a GeometryCollection with PlanarCells, and add each cut cell back to the GeometryCollection as a new child of their source Geometry.  For geometries that would not be cut, nothing is added.
  *
  * @param Cells				Defines the cutting planes and division of space
- * @param Collection			The collection to be cut
- * @param GeometryIndices	Which geometry groups inside the collection to cut
+ * @param Collection		The collection to be cut
+ * @param TransformIndices	Which transform groups inside the collection to cut
  * @param TransformCells	Optional transform of the planar cut; if unset, defaults to Identity
  * @param bIncludeOutsideCellInOutput	If true, geometry that was not inside any of the cells (e.g. was outside of the bounds of all cutting geometry) will still be included in the output; if false, it will be discarded.
  * @param CheckDistanceAcrossOutsideCellForProximity	If > 0, when a plane is neighboring the "outside" cell, instead of setting proximity to the outside cell, the algo will sample a point this far outside the cell in the normal direction of the plane to see if there is actually a non-outside cell there.  (Useful for bricks w/out mortar)
@@ -180,7 +180,7 @@ int32 PLANARCUT_API CutWithPlanarCells(
 int32 PLANARCUT_API CutMultipleWithPlanarCells(
 	FPlanarCells &Cells,
 	FGeometryCollection& Collection,
-	const TArrayView<const int32>& GeometryIndices,
+	const TArrayView<const int32>& TransformIndices,
 	const TOptional<FTransform>& TransformCells = TOptional<FTransform>(),
 	bool bIncludeOutsideCellInOutput = true,
 	float CheckDistanceAcrossOutsideCellForProximity = 0,
@@ -194,7 +194,7 @@ int32 PLANARCUT_API CutMultipleWithPlanarCells(
  * @param Planes				Defines the cutting planes and division of space
  * @param InternalSurfaceMaterials	Defines material properties for any added internal surfaces
  * @param Collection			The collection to be cut
- * @param GeometryIndices	Which geometry groups inside the collection to cut
+ * @param TransformIndices	Which transform groups inside the collection to cut
  * @param TransformCells	Optional transform of the planar cut; if unset, defaults to Identity
  * @param bFlattenToSingleLayer If true, only add one layer overall to the hierarchy even if some geometry is cut multiple times.
  * @param VertexInterpolate	Function that interpolates vertex properties (UVs, normals, etc); a default that handles all the normal vertex properties is provided, should only need to replace this if you have custom attributes
@@ -204,7 +204,7 @@ int32 PLANARCUT_API CutMultipleWithMultiplePlanes(
 	const TArrayView<const FPlane>& Planes,
 	FInternalSurfaceMaterials& InternalSurfaceMaterials,
 	FGeometryCollection& Collection,
-	const TArrayView<const int32>& GeometryIndices,
+	const TArrayView<const int32>& TransformIndices,
 	const TOptional<FTransform>& TransformCells = TOptional<FTransform>(),
 	bool bFlattenToSingleLayer = true,
 	bool bSetDefaultInternalMaterialsFromCollection = true,
