robustness fixes to planarcuts based on rtl feedback; unlikely non-fatal cases should not be handled with check  #jira none  #rb none  #rnx

diff --git a/Engine/Plugins/Experimental/GeometryProcessing/Source/GeometryAlgorithms/Public/Arrangement2d.h b/Engine/Plugins/Experimental/GeometryProcessing/Source/GeometryAlgorithms/Public/Arrangement2d.h
--- a/Engine/Plugins/Experimental/GeometryProcessing/Source/GeometryAlgorithms/Public/Arrangement2d.h
+++ b/Engine/Plugins/Experimental/GeometryProcessing/Source/GeometryAlgorithms/Public/Arrangement2d.h
@@ -311,7 +311,7 @@ protected:
 					FSegment2d new_seg = Graph.GetEdgeSegment(new_eid);
 					FVector2d p1 = Intr.Intr.GetSegment1().PointAt(t1);
 					double new_t1 = new_seg.Project(p1);
-					check(new_t1 <= FMath::Abs(new_seg.Extent));
+					// note: new_t1 may be outside of new_seg due to snapping; in this case the segment will just not be split
 
 					FIndex2i new_info = split_segment_at_t(new_eid, new_t1, VertexSnapTol);
 					FVector2d v = Graph.GetVertex(new_info.A);
diff --git a/Engine/Plugins/Experimental/PlanarCutPlugin/Source/PlanarCut/Private/PlanarCut.cpp b/Engine/Plugins/Experimental/PlanarCutPlugin/Source/PlanarCut/Private/PlanarCut.cpp
--- a/Engine/Plugins/Experimental/PlanarCutPlugin/Source/PlanarCut/Private/PlanarCut.cpp
+++ b/Engine/Plugins/Experimental/PlanarCutPlugin/Source/PlanarCut/Private/PlanarCut.cpp
@@ -342,7 +342,7 @@ FPlanarCells::FPlanarCells(const FBox &Region, const FIntVector& CubesPerAxis)
 			for (int32 Xi = 0; Xi < VertsPerAxis.X; Xi++)
 			{
 				PlaneBoundaryVertices[VertIdx] = Region.Min + FVector(Xi * CellSizes.X, Yi * CellSizes.Y, Zi * CellSizes.Z);
-				check(VertIdx == ToIdxUnsafe(VertsPerAxis, Xi, Yi, Zi));
+				ensure(VertIdx == ToIdxUnsafe(VertsPerAxis, Xi, Yi, Zi));
 				VertIdx++;
 			}
 		}
@@ -542,8 +542,8 @@ FPlanarCells::FPlanarCells(const FBox &Region, const TArrayView<const FColor> Im
 
 	for (int32 CellIdx = 0; CellIdx < NumCells; CellIdx++)
 	{
-		check(CellBoundaryCorners[CellIdx].Num() > 0); // there must not be any regions with no boundary
-		check(PerCellBoundaryEdgeArrays[CellIdx].Num() == 0); // all boundary edge array should have been consumed and turned to full boundary loops
+		ensure(CellBoundaryCorners[CellIdx].Num() > 0); // there must not be any regions with no boundary
+		ensure(PerCellBoundaryEdgeArrays[CellIdx].Num() == 0); // all boundary edge array should have been consumed and turned to full boundary loops
 		ensureMsgf(CellBoundaryCorners[CellIdx].Num() == 1, TEXT("Have not implemented support for regions with holes!"));
 
 		int32 BoundaryStart = PlaneBoundaryVertices.Num();
@@ -732,8 +732,7 @@ struct OutputCells
 		else
 		{
 			// cells should be symmetrically connected, so if the cell was already connected A->B, it should have been connected B->A
-			// here's a checkSlow just to be sure
-			checkSlow(INDEX_NONE != NeighborCells[CellB].Find(CellA));
+			ensure(INDEX_NONE != NeighborCells[CellB].Find(CellA));
 		}
 	}
 
@@ -804,7 +803,7 @@ struct OutputCells
 		if (!bIncludeOutsideCellInOutput && NoCellIdx > -1)
 		{
 			NumCellsToDump--;
-			check(NoCellIdx == NumCellsToDump); // by convention the index of the "no cell" geometry is the last index
+			ensure(NoCellIdx == NumCellsToDump); // by convention the index of the "no cell" geometry is the last index
 			if (CellTriangles[NoCellIdx].Num())
 			{
 				NumNewGeometries--;
@@ -1229,7 +1228,7 @@ void CutWithPlanarCellsHelper(
 					CrossPosns[1] = Vertices[MinSDIdx];
 					CrossIdx++;
 				}
-				check(CrossIdx == 2);
+				ensure(CrossIdx == 2);
 				EdgesOnPlane[PlaneIdx].Add(TPair<FVector, FVector>(CrossPosns[0], CrossPosns[1]));
 			}
 			else {
@@ -1257,7 +1256,7 @@ void CutWithPlanarCellsHelper(
 					AddOn(SX, Tri.X);
 					AddOn(SY, Tri.Y);
 					AddOn(SZ, Tri.Z);
-					check(OnIdx == 2);
+					ensure(OnIdx == 2);
 					EdgesOnPlane[PlaneIdx].Add(TPair<FVector, FVector>(OnPosns[0], OnPosns[1]));
 				}
 			}
@@ -1591,7 +1590,7 @@ void CutWithPlanarCellsHelper(
 								double T = SDA / (SDA - SDB);
 								FVector2d OnBoundary = PlanarEdges[EdgeIdx].Key * (1 - T) + PlanarEdges[EdgeIdx].Value * T;
 								double SDO = (OnBoundary - Pt).Dot(EdgeNormal);
-								check(FMath::Abs(SDO) < 1e-4);
+								ensure(FMath::Abs(SDO) < 1e-4);
 								if (SDA < 0)
 								{
 									PlanarEdges[EdgeIdx].Key = OnBoundary;
@@ -1657,15 +1656,16 @@ void CutWithPlanarCellsHelper(
 			if (bNoiseOnPlane)
 			{
 				const FNoiseSettings& Noise = InternalMaterials->NoiseSettings.GetValue();
-				check(Noise.PointSpacing > 0);
+				const float MinPointSpacing = .1;
+				float PointSpacing = FMath::Max(MinPointSpacing, Noise.PointSpacing);
 
 				// make a new point hash for blue noise point location queries
 				// this is essentially the same as the point hash in arrangement2d but with cell spacing set based on the point spacing; the arrangement2d one can have a way-too-small point spacing!
-				TPointHashGrid2d<int> NoisePointHash(Noise.PointSpacing, -1);
+				TPointHashGrid2d<int> NoisePointHash(PointSpacing, -1);
 				auto HasVertexNear = [&](const FVector2d& V)
 				{
 					auto FuncDistSq = [&](int B) { return V.DistanceSquared(Arrangement.Graph.GetVertex(B)); };
-					TPair<int, double> NearestPt = NoisePointHash.FindNearestInRadius(V, Noise.PointSpacing*.99, FuncDistSq);
+					TPair<int, double> NearestPt = NoisePointHash.FindNearestInRadius(V, PointSpacing*.99, FuncDistSq);
 					return NearestPt.Key != NoisePointHash.InvalidValue();
 				};
 				for (int32 VertIdx = 0; VertIdx < Arrangement.Graph.MaxVertexID(); VertIdx++)
@@ -1676,7 +1676,7 @@ void CutWithPlanarCellsHelper(
 					}
 				}
 
-				double SpacingSq = Noise.PointSpacing*Noise.PointSpacing;
+				double SpacingSq = PointSpacing*PointSpacing;
 				for (int EdgeIdx : Arrangement.Graph.EdgeIndices())
 				{
 					FDynamicGraph::FEdge Edge = Arrangement.Graph.GetEdge(EdgeIdx);
@@ -1709,13 +1709,13 @@ void CutWithPlanarCellsHelper(
 						
 					}
 				}
-				for (double X = Bounds2D.Min.X; X < Bounds2D.Max.X; X += Noise.PointSpacing)
+				for (double X = Bounds2D.Min.X; X < Bounds2D.Max.X; X += PointSpacing)
 				{
-					for (double Y = Bounds2D.Min.Y; Y < Bounds2D.Max.Y; Y += Noise.PointSpacing)
+					for (double Y = Bounds2D.Min.Y; Y < Bounds2D.Max.Y; Y += PointSpacing)
 					{
 						for (int Attempt = 0; Attempt < 5; Attempt++)
 						{
-							FVector2d Pt(X + FMath::FRand() * Noise.PointSpacing*.5, Y + FMath::FRand() * Noise.PointSpacing*.5);
+							FVector2d Pt(X + FMath::FRand() * PointSpacing*.5, Y + FMath::FRand() * PointSpacing*.5);
 							if (!HasVertexNear(Pt))
 							{
 								int PtIdx = Arrangement.Insert(Pt);
@@ -1792,7 +1792,7 @@ void CutWithPlanarCellsHelper(
 
 			ensure(SkippedEdges.Num() == 0); // TODO: remove this ensure; for now I'm just curious how much triangulation fails in practice
 
-			checkSlow(Arrangement.Graph.IsCompact());
+			ensure(Arrangement.Graph.IsCompact());
 			for (int32 VertIdx = 0; VertIdx < Arrangement.Graph.MaxVertexID(); VertIdx++)
 			{
 				Triangulation.LocalVertices.Add(PlaneFrames[PlaneIdx].UnProject(Arrangement.Graph.GetVertex(VertIdx)));
@@ -1843,7 +1843,7 @@ void CutWithPlanarCellsHelper(
 			if (NumBoundary > 2) // if there are at least 3 boundary points, we still have something we could triangulate
 			{
 				// TODO: should actually be "AssumeConvexFacets"?
-				check(Cells.AssumeConvexCells);  // should have followed the above code path w/ triangulation if non-convex faces
+				ensure(Cells.AssumeConvexCells);  // should have followed the above code path w/ triangulation if non-convex faces
 
 				FVector FacetCentroid(0, 0, 0);
 				for (int32 VertIdx : BoundaryIndices)
@@ -2426,7 +2426,7 @@ int32 CutMultipleWithMultiplePlanes(
 						}
 						break;
 						default:
-							check(false); // PlaneSide must be -1, 0, or 1
+							ensure(false); // PlaneSide must be -1, 0, or 1
 						}
 
 					}
