fix static analysis issue for plane cut helper function #jira UE-78141
#rnx

diff --git a/Engine/Plugins/Experimental/PlanarCutPlugin/Source/PlanarCut/Private/PlanarCut.cpp b/Engine/Plugins/Experimental/PlanarCutPlugin/Source/PlanarCut/Private/PlanarCut.cpp
--- a/Engine/Plugins/Experimental/PlanarCutPlugin/Source/PlanarCut/Private/PlanarCut.cpp
+++ b/Engine/Plugins/Experimental/PlanarCutPlugin/Source/PlanarCut/Private/PlanarCut.cpp
@@ -1245,17 +1245,17 @@ void CutWithPlanarCellsHelper(
 					int32 OnIdx = 0;
 					FVector OnPosns[2];
 
-					auto AddOn = [&](float SD, int32 VIdx)
+					auto AddOn = [&Vertices, &PlaneEps, &OnPosns](float SD, int32 VIdx, int32 &OnIdxRef)
 					{
-						if (OnIdx < 2 && FMath::Abs(SD) < PlaneEps)
+						if (OnIdxRef < 2 && FMath::Abs(SD) < PlaneEps)
 						{
-							OnPosns[OnIdx] = Vertices[VIdx];
-							OnIdx++;
+							OnPosns[OnIdxRef] = Vertices[VIdx];
+							OnIdxRef++;
 						}
 					};
-					AddOn(SX, Tri.X);
-					AddOn(SY, Tri.Y);
-					AddOn(SZ, Tri.Z);
+					AddOn(SX, Tri.X, OnIdx);
+					AddOn(SY, Tri.Y, OnIdx);
+					AddOn(SZ, Tri.Z, OnIdx);
 					ensure(OnIdx == 2);
 					EdgesOnPlane[PlaneIdx].Add(TPair<FVector, FVector>(OnPosns[0], OnPosns[1]));
 				}
