fix robustness issues for fracturing large scale meshes

#rnx

diff --git a/Engine/Plugins/Experimental/PlanarCutPlugin/Source/PlanarCut/Private/PlanarCut.cpp b/Engine/Plugins/Experimental/PlanarCutPlugin/Source/PlanarCut/Private/PlanarCut.cpp
--- a/Engine/Plugins/Experimental/PlanarCutPlugin/Source/PlanarCut/Private/PlanarCut.cpp
+++ b/Engine/Plugins/Experimental/PlanarCutPlugin/Source/PlanarCut/Private/PlanarCut.cpp
@@ -149,21 +149,25 @@ void FInternalSurfaceMaterials::SetUVScaleFromCollection(const FGeometryCollecti
 	{
 		GlobalUVScale =  UVDistance / WorldDistance;
 	}
-	
 }
 
+inline double PlaneDotDouble(const FPlane& Plane, const FVector& V)
+{
+	return (double)Plane.X*V.X + Plane.Y*V.Y + Plane.Z*V.Z - Plane.W;
+}
 
-inline int PlaneSide(const FPlane &Plane, const FVector& V, const float Epsilon = 1e-3)
+
+inline int PlaneSide(const FPlane &Plane, const FVector& V, const double Epsilon = 1e-3)
 {
-	float SD = Plane.PlaneDot(V);
+	double SD = PlaneDotDouble(Plane, V);
 	return SD > Epsilon ? 1 : SD < -Epsilon ? -1 : 0;
 }
 
 // TODO: warning -- If Epsilon is too small, we can hit an infinite loop on mesh cutting (if the edge cut is still seen as crossing!)
-inline bool IsSegmentCrossing(const FPlane &Plane, const FVector& A, const FVector& B, float& CrossingT, const float Epsilon = 1e-3)
+inline bool IsSegmentCrossing(const FPlane &Plane, const FVector& A, const FVector& B, double& CrossingT, const double Epsilon = 1e-3)
 {
-	float SDA = Plane.PlaneDot(A);
-	float SDB = Plane.PlaneDot(B);
+	double SDA = PlaneDotDouble(Plane, A);
+	double SDB = PlaneDotDouble(Plane, B);
 	CrossingT = SDA / (SDA - SDB);
 	int32 ASide = SDA < -Epsilon ? -1 : SDA > Epsilon ? 1 : 0;
 	int32 BSide = SDB < -Epsilon ? -1 : SDB > Epsilon ? 1 : 0;
@@ -178,7 +182,7 @@ FPlanarCells::FPlanarCells(const FPlane& P)
 
 	CellFromPosition = TFunction<int32(FVector)>([P](FVector Position)
 	{
-		return P.PlaneDot(Position) > 0 ? 1 : 0;
+		return PlaneDotDouble(P, Position) > 0 ? 1 : 0;
 	});
 }
 
@@ -987,6 +991,69 @@ struct OutputCells
 	}
 };
 
+/**
+* Transform local geometry, updating the corresponding transform so the shape itself is not changed
+*/
+void TransformLocalGeometry(FGeometryCollection& Source, int32 TransformIdx, const FTransform& Transform, const FTransform& InverseTransform)
+{
+	int GeometryIdx = Source.TransformToGeometryIndex[TransformIdx];
+	// recompute bounds (rather than directly transforming, so it remains a tight bound even if transform includes rotation)
+	FBox Bounds; Bounds.Init();
+	for (int32 VertIdx = Source.VertexStart[GeometryIdx], VertEnd = Source.VertexCount[GeometryIdx] + Source.VertexStart[GeometryIdx]; VertIdx < VertEnd; VertIdx++)
+	{
+		FVector Pos = Transform.TransformPosition(Source.Vertex[VertIdx]);
+		Bounds += Pos;
+		Source.Vertex[VertIdx] = Pos;
+	}
+	Source.BoundingBox[GeometryIdx] = Bounds;
+	Source.Transform[TransformIdx] = InverseTransform * Source.Transform[TransformIdx];
+}
+
+/** 
+ * Update a given transform w/ geometry to have vertices locally centered at the origin and positions not exceeding a -1,1 range, changing the transform appropriately so that the shape itself is not changed
+ */
+FTransform CenterAndScaleLocalGeometry(FGeometryCollection& Source, int32 TransformIdx)
+{
+	int GeometryIdx = Source.TransformToGeometryIndex[TransformIdx];
+	if (!ensure(GeometryIdx != INDEX_NONE)) // transform had no geometry
+	{
+		return FTransform::Identity;
+	}
+
+	FBox GeomBox; GeomBox.Init();
+	if (Source.BoundingBox.Num())
+	{
+		GeomBox = Source.BoundingBox[GeometryIdx];
+	}
+	if (!GeomBox.IsValid)
+	{
+		for (int32 VertIdx = Source.VertexStart[GeometryIdx], VertEnd = Source.VertexCount[GeometryIdx] + Source.VertexStart[GeometryIdx]; VertIdx < VertEnd; VertIdx++)
+		{
+			GeomBox += Source.Vertex[VertIdx];
+		}
+	}
+
+	if (!ensure(GeomBox.IsValid)) // transform had corresponding geometry index but it had zero vertices?
+	{
+		return FTransform::Identity;
+	}
+	
+	FVector Center, Extents;
+	GeomBox.GetCenterAndExtents(Center, Extents);
+	float MaxExtent = Extents.GetMax();
+	float InvScaleFactor = MaxExtent < 1 ? 1 : MaxExtent;
+	float ScaleFactor = 1.0f / InvScaleFactor;
+	FTransform CenterAndFit, InverseCenterAndFit;
+	CenterAndFit.SetTranslationAndScale3D(-Center*ScaleFactor, FVector(ScaleFactor, ScaleFactor, ScaleFactor));
+	InverseCenterAndFit.SetTranslationAndScale3D(Center, FVector(InvScaleFactor, InvScaleFactor, InvScaleFactor));
+
+	TransformLocalGeometry(Source, TransformIdx, CenterAndFit, InverseCenterAndFit);
+
+	return InverseCenterAndFit;
+}
+
+
+
 //// useful for debugging: code to dump a single geometry to obj
 //#include <fstream>
 //#include <string>
@@ -1042,7 +1109,7 @@ void CutWithPlanarCellsHelper(
 	int32 TriangleStart,
 	int32 NumTriangles,
 	const TArrayView<const FVector> TriangleNormals,
-	float PlaneEps,
+	double PlaneEps,
 	float CheckDistanceAcrossOutsideCellForProximity,
 	TFunction<void(const FGeometryCollection&, int32, const FGeometryCollection&, int32, float, int32, FGeometryCollection&)> Interpolate,
 	OutputCells &Output,
@@ -1060,7 +1127,18 @@ void CutWithPlanarCellsHelper(
 
 	// consider trade-offs between cases where we could have a more consistent mesh vs having simpler processing / fewer triangles
 	constexpr bool bCareAboutTJunctionsEvenALittleBit = false;
-	
+	bool bNoiseOnPlane = InternalMaterials->NoiseSettings.IsSet();
+
+	// extract an average scale for this transform to support properly spacing noise points, if requested
+	float AverageGlobalScale = 1;
+	if (bNoiseOnPlane)
+	{
+		FTransform LocalToGlobalTransform = GeometryCollectionAlgo::GlobalMatrix(Source.Transform, Source.Parent, Source.TransformIndex[GeometryIdx]);
+		FVector Scales = LocalToGlobalTransform.GetScale3D();
+		AverageGlobalScale = FMath::Max(KINDA_SMALL_NUMBER, FVector::DotProduct(Scales.GetAbs(), FVector(1. / 3.)));
+	}
+	float AverageGlobalScaleInv = 1.0 / AverageGlobalScale;
+
 	struct PlaneFrame
 	{
 		FVector3d Origin, X, Y;
@@ -1157,7 +1235,9 @@ void CutWithPlanarCellsHelper(
 			}
 		}
 
-		FAxisAlignedBox3d PlaneFacetBoxPlusEps(BoundingBox.ExpandBy(PlaneEps));
+		FAxisAlignedBox3d PlaneFacetBoxPlusEps(BoundingBox);
+		PlaneFacetBoxPlusEps.Max = PlaneFacetBoxPlusEps.Max + PlaneEps;
+		PlaneFacetBoxPlusEps.Min = PlaneFacetBoxPlusEps.Min - PlaneEps;
 		TMeshAABBTree3<FGeometryCollectionMeshAdapter>::FTreeTraversal TraverseNearPlane;
 		TraverseNearPlane.NextBoxF = [&](const FAxisAlignedBox3d& Box, int Depth)
 		{
@@ -1186,7 +1266,7 @@ void CutWithPlanarCellsHelper(
 
 			// TODO: further filter out cases where triangle doesn't intersect planar facet
 
-			float SX = Plane.PlaneDot(Vertices[Tri.X]), SY = Plane.PlaneDot(Vertices[Tri.Y]), SZ = Plane.PlaneDot(Vertices[Tri.Z]);
+			double SX = PlaneDotDouble(Plane,Vertices[Tri.X]), SY = PlaneDotDouble(Plane, Vertices[Tri.Y]), SZ = PlaneDotDouble(Plane, Vertices[Tri.Z]);
 			int32 SXSide = SX < -PlaneEps ? -1 : SX > PlaneEps ? 1 : 0;
 			int32 SYSide = SY < -PlaneEps ? -1 : SY > PlaneEps ? 1 : 0;
 			int32 SZSide = SZ < -PlaneEps ? -1 : SZ > PlaneEps ? 1 : 0;
@@ -1201,7 +1281,7 @@ void CutWithPlanarCellsHelper(
 				{
 					if (CrossIdxRef < 2 && ASide*BSide == -1)
 					{
-						float T = SDA / (SDA - SDB);
+						double T = SDA / (SDA - SDB);
 						CrossPosns[CrossIdxRef] = (1 - T) * Vertices[AIdx] + T * Vertices[BIdx];
 						CrossIdxRef++;
 					}
@@ -1316,7 +1396,7 @@ void CutWithPlanarCellsHelper(
 				{
 					for (int32 Idx = 0, LastIdx = PlaneBoundaryNum - 1; Idx < PlaneBoundaryNum; LastIdx = Idx++)
 					{
-						float T;
+						double T;
 						const FVector& A = TransformedPlaneBoundaryVertices[PlaneBoundary[Idx]];
 						const FVector& B = TransformedPlaneBoundaryVertices[PlaneBoundary[LastIdx]];
 						IsSegmentCrossing(TriPlane, A, B, T);
@@ -1382,7 +1462,7 @@ void CutWithPlanarCellsHelper(
 					//}
 
 					{
-						float t;
+						double t;
 						FVector P0 = VertexPos(V0), P1 = VertexPos(V1);
 						if (IsSegmentCrossing(TransformedPlanes[ConsiderPlaneIdx], P0, P1, t))
 						{
@@ -1415,7 +1495,7 @@ void CutWithPlanarCellsHelper(
 						FVector A = VertexPos(Tri.X), B = VertexPos(Tri.Y), C = VertexPos(Tri.Z);
 						auto ContainCrossing = [&Plane, &PlaneEps, &IntersectionDirection](FInterval1d& Inter, const FVector& P0, const FVector& P1)
 						{
-							float T;
+							double T;
 							IsSegmentCrossing(Plane, P0, P1, T);
 							if (T > -PlaneEps && T < 1 + PlaneEps)
 							{
@@ -1499,7 +1579,6 @@ void CutWithPlanarCellsHelper(
 
 		// check if constrained Delaunay triangulation problem needed for plane (false if no geometry was touching the planar facet)
 		bool bAnyElementsOnPlane = EdgesOnPlane[PlaneIdx].Num() + TrianglesOnPlane[PlaneIdx].Num() > 0;
-		bool bNoiseOnPlane = InternalMaterials->NoiseSettings.IsSet();
 		bool bConvexFacet = Cells.AssumeConvexCells;
 		bool bHasBoundary = NumBoundary > 2;
 		if (bAnyElementsOnPlane || (bHasBoundary && (bNoiseOnPlane || !bConvexFacet)))
@@ -1640,14 +1719,17 @@ void CutWithPlanarCellsHelper(
 			
 
 			const double ArrangementTol = 1e-4;
-			FArrangement2d Arrangement(FMath::Max(Bounds2D.MaxDim() / 64, ArrangementTol * 10));
+			const double ScaleF = 1.0 / FMathd::Max(.01, Bounds2D.MaxDim());
+			const FVector2d Offset = -Bounds2D.Center();
+			FAxisAlignedBox2d ScaledBounds2D((Bounds2D.Min + Offset)*ScaleF, (Bounds2D.Max + Offset)*ScaleF);
+			FArrangement2d Arrangement(FMath::Max(ScaledBounds2D.MaxDim() / 64, ArrangementTol * 10));
 			Arrangement.VertexSnapTol = ArrangementTol;
 			int32 BoundaryEdgeGroupID = -1;
 			for (int32 EdgeIdx = PlanarEdges.Num()-1; EdgeIdx >= 0; EdgeIdx--)
 			{
 				const TPair<FVector2d, FVector2d>& Edge2d = PlanarEdges[EdgeIdx];
 				int32 EdgeGroupID = EdgeIdx >= BoundaryEdgeStart ? BoundaryEdgeGroupID : EdgeIdx; // give all boundary edges the same group ID
-				Arrangement.Insert(Edge2d.Key, Edge2d.Value, EdgeGroupID);
+				Arrangement.Insert((Edge2d.Key+Offset)*ScaleF, (Edge2d.Value+Offset)*ScaleF, EdgeGroupID);
 			}
 			TArray<int32> SkippedEdges;
 			TArray<FIntVector> PlaneTriangulation;
@@ -1657,7 +1739,7 @@ void CutWithPlanarCellsHelper(
 			{
 				const FNoiseSettings& Noise = InternalMaterials->NoiseSettings.GetValue();
 				const float MinPointSpacing = .1;
-				float PointSpacing = FMath::Max(MinPointSpacing, Noise.PointSpacing);
+				float PointSpacing = FMath::Max(MinPointSpacing, Noise.PointSpacing*float(ScaleF) * AverageGlobalScaleInv);
 
 				// make a new point hash for blue noise point location queries
 				// this is essentially the same as the point hash in arrangement2d but with cell spacing set based on the point spacing; the arrangement2d one can have a way-too-small point spacing!
@@ -1709,9 +1791,9 @@ void CutWithPlanarCellsHelper(
 						
 					}
 				}
-				for (double X = Bounds2D.Min.X; X < Bounds2D.Max.X; X += PointSpacing)
+				for (double X = ScaledBounds2D.Min.X; X < ScaledBounds2D.Max.X; X += PointSpacing)
 				{
-					for (double Y = Bounds2D.Min.Y; Y < Bounds2D.Max.Y; Y += PointSpacing)
+					for (double Y = ScaledBounds2D.Min.Y; Y < ScaledBounds2D.Max.Y; Y += PointSpacing)
 					{
 						for (int Attempt = 0; Attempt < 5; Attempt++)
 						{
@@ -1730,6 +1812,13 @@ void CutWithPlanarCellsHelper(
 
 			Arrangement.AttemptTriangulate(PlaneTriangulation, SkippedEdges, BoundaryEdgeGroupID);
 
+			// undo scaling
+			double InvScaleF = 1.0 / ScaleF;
+			for (int GraphVertIdx : Arrangement.Graph.VertexIndices())
+			{
+				Arrangement.Graph.SetVertex(GraphVertIdx, (Arrangement.Graph.GetVertex(GraphVertIdx) * InvScaleF) - Offset);
+			}
+
 			// Eat any triangles that are inside coplanar triangles on the face
 			// TODO: optimize this? currently doing this in the simplest way possible because it is likely a rare case
 			
@@ -1831,13 +1920,13 @@ void CutWithPlanarCellsHelper(
 					float OctaveScale = 1;
 					for (int32 Octave = 0; Octave < Octaves; Octave++, OctaveScale *= 2)
 					{
-						NoiseValue += FMath::PerlinNoise2D(V * OctaveScale) / OctaveScale;
+						NoiseValue += FMath::PerlinNoise2D(V * OctaveScale * AverageGlobalScale) / OctaveScale;
 					}
-					Triangulation.LocalVertices[VertexIdx] += Z * NoiseValue;
+					Triangulation.LocalVertices[VertexIdx] += Z * NoiseValue * AverageGlobalScaleInv;
 				}
 			}
 		}
-		else
+		else // no CDT needed; just triangulate the cell directly
 		{
 			ensure(NumBoundary != 1 && NumBoundary != 2);  // Point or Line Segment boundaries would be weird / imply a bug maybe
 			if (NumBoundary > 2) // if there are at least 3 boundary points, we still have something we could triangulate
@@ -2015,20 +2104,30 @@ void CutWithPlanarCellsHelper(
 
 void TransformPlanes(const FTransform& Transform, const FPlanarCells& Ref, TArray<FPlane>& Planes, TArray<FVector>& PlaneBoundaries)
 {
-	FMatrix Matrix = Transform.ToMatrixWithScale();
-	const FMatrix TransposeAdjoint = Matrix.TransposeAdjoint();
-	const float DetM = Matrix.Determinant();
-
+	// Note: custom implementation of normal transform for robustness, especially to ensure we don't zero the normals for significantly scaled geometry
+	FTransform NormalTransform = Transform;
+	FVector3d ScaleVec(NormalTransform.GetScale3D());
+	double ScaleDetSign = FMathd::SignNonZero(ScaleVec.X) * FMathd::SignNonZero(ScaleVec.Y) * FMathd::SignNonZero(ScaleVec.Z);
+	double ScaleMaxAbs = ScaleVec.MaxAbs();
+	if (ScaleMaxAbs > DBL_MIN)
+	{
+		ScaleVec /= ScaleMaxAbs;
+	}
+	FVector3d NormalScale(ScaleVec.Y*ScaleVec.Z*ScaleDetSign, ScaleVec.X*ScaleVec.Z*ScaleDetSign, ScaleVec.X*ScaleVec.Y*ScaleDetSign);
+	NormalTransform.SetScale3D(NormalScale);
+	
 	Planes.SetNum(Ref.Planes.Num());
 	for (int32 PlaneIdx = 0, PlanesNum = Planes.Num(); PlaneIdx < PlanesNum; PlaneIdx++)
 	{
-		Planes[PlaneIdx] = Ref.Planes[PlaneIdx].TransformByUsingAdjointT(Matrix, DetM, TransposeAdjoint);
+		FVector Pos = Transform.TransformPosition(Ref.Planes[PlaneIdx] * Ref.Planes[PlaneIdx].W);
+		FVector Normal = NormalTransform.TransformVector(Ref.Planes[PlaneIdx]).GetSafeNormal(FLT_MIN);
+		Planes[PlaneIdx] = FPlane(Pos, Normal);
 	}
 
 	PlaneBoundaries.SetNum(Ref.PlaneBoundaryVertices.Num());
 	for (int32 VertIdx = 0, VertsNum = PlaneBoundaries.Num(); VertIdx < VertsNum; VertIdx++)
 	{
-		PlaneBoundaries[VertIdx] = Matrix.TransformPosition(Ref.PlaneBoundaryVertices[VertIdx]);
+		PlaneBoundaries[VertIdx] = Transform.TransformPosition(Ref.PlaneBoundaryVertices[VertIdx]);
 	}
 }
 
