first version of mesh boolean-based fracture backend
close to parity with the old version; both versions are left in and a console variable is provided to toggle back, to enable artist testing

the new version has different (hopefully better) noise and tries to avoid creating 'cracks' in the mesh, but should otherwise behave similarly.

diff --git a/Engine/Plugins/Experimental/ChaosEditor/Source/FractureEditor/Private/FractureToolVoronoiBase.cpp b/Engine/Plugins/Experimental/ChaosEditor/Source/FractureEditor/Private/FractureToolVoronoiBase.cpp
--- a/Engine/Plugins/Experimental/ChaosEditor/Source/FractureEditor/Private/FractureToolVoronoiBase.cpp
+++ b/Engine/Plugins/Experimental/ChaosEditor/Source/FractureEditor/Private/FractureToolVoronoiBase.cpp
@@ -88,11 +88,17 @@ void UFractureToolVoronoiBase::ExecuteFracture(const FFractureContext& FractureC
 		TSharedPtr<FGeometryCollection, ESPMode::ThreadSafe> GeometryCollectionPtr = FractureContext.FracturedGeometryCollection->GetGeometryCollection();
 		if (FGeometryCollection* GeometryCollection = GeometryCollectionPtr.Get())
 		{
+			const UFractureCommonSettings* LocalCommonSettings = GetDefault<UFractureCommonSettings>();
+
 			TArray<FVector> Sites;
 			GenerateVoronoiSites(FractureContext, Sites);
-			FVoronoiDiagram Voronoi(Sites, FractureContext.Bounds, .1f);
-
-			const UFractureCommonSettings* LocalCommonSettings = GetDefault<UFractureCommonSettings>();
+			FBox VoronoiBounds = FractureContext.Bounds;
+			if (LocalCommonSettings->Amplitude > 0.0f)
+			{
+				// expand bounds to make sure noise-perturbed voronoi cells still contain the whole input mesh
+				VoronoiBounds = VoronoiBounds.ExpandBy(LocalCommonSettings->Amplitude);
+			}
+			FVoronoiDiagram Voronoi(Sites, VoronoiBounds, .1f);
 
 			FPlanarCells VoronoiPlanarCells = FPlanarCells(Sites, Voronoi);
 
diff --git a/Engine/Plugins/Experimental/PlanarCutPlugin/Source/PlanarCut/PlanarCut.Build.cs b/Engine/Plugins/Experimental/PlanarCutPlugin/Source/PlanarCut/PlanarCut.Build.cs
--- a/Engine/Plugins/Experimental/PlanarCutPlugin/Source/PlanarCut/PlanarCut.Build.cs
+++ b/Engine/Plugins/Experimental/PlanarCutPlugin/Source/PlanarCut/PlanarCut.Build.cs
@@ -16,6 +16,7 @@ namespace UnrealBuildTool.Rules
 					"CoreUObject",
                     "Voronoi",
 					"GeometricObjects",
+					"DynamicMesh",
 					"Chaos",
 					"GeometryAlgorithms"
                 }
diff --git a/Engine/Plugins/Experimental/PlanarCutPlugin/Source/PlanarCut/Private/PlanarCut.cpp b/Engine/Plugins/Experimental/PlanarCutPlugin/Source/PlanarCut/Private/PlanarCut.cpp
--- a/Engine/Plugins/Experimental/PlanarCutPlugin/Source/PlanarCut/Private/PlanarCut.cpp
+++ b/Engine/Plugins/Experimental/PlanarCutPlugin/Source/PlanarCut/Private/PlanarCut.cpp
@@ -4,6 +4,7 @@
 
 #include "Async/ParallelFor.h"
 #include "Spatial/FastWinding.h"
+#include "Spatial/PointHashGrid3.h"
 #include "Arrangement2d.h"
 #include "MeshAdapter.h"
 #include "FrameTypes.h"
@@ -13,11 +14,34 @@
 #include "GeometryCollection/GeometryCollectionClusteringUtility.h"
 
 
+#include "DynamicMesh3.h"
+#include "DynamicMeshEditor.h"
+#include "DynamicMeshAABBTree3.h"
+#include "Operations/MeshBoolean.h"
+#include "Operations/MeshSelfUnion.h"
+#include "Operations/MergeCoincidentMeshEdges.h"
+#include "MeshBoundaryLoops.h"
+#include "QueueRemesher.h"
+#include "DynamicVertexAttribute.h"
+#include "MeshNormals.h"
+#include "ConstrainedDelaunay2.h"
+
+#include "Algo/Rotate.h"
+
 #if WITH_EDITOR
 #include "Misc/ScopedSlowTask.h"
 #endif
 
+// Console variable to toggle between new and old fracture method; may delete this when we delete the old method
+static TAutoConsoleVariable<int32> CVarFractureMethod(
+	TEXT("Fracture.MeshCuttingMethod"),
+	1,
+	TEXT("Defines the which method is used under the hood to compute fracture geometry.\n")
+	TEXT("<=0: original, with fake noise\n")
+	TEXT("  1: experimental new method\n"),
+	ECVF_Scalability);
 
+
 /**
  * Adapter that lets the generic fast winding and AABB tree code view a geometry from a geometrycollection as a simple index buffer triangle mesh
  */
@@ -79,9 +103,1164 @@
 	}
 };
 
+// FDynamicMesh3 + additional attributes as needed to carry standard geometry collection attributes
+struct FAugmentedDynamicMesh
+{
+	FDynamicMesh3 Mesh;
+	const FName TangentUAttribName = "TangentUAttrib";
+	const FName TangentVAttribName = "TangentVAttrib";
+	const FName VisibleAttribName = "VisibleAttrib";
 
+	FAugmentedDynamicMesh()
+	{
+		Mesh.EnableAttributes();
+		Mesh.Attributes()->EnableMaterialID();
+		Mesh.Attributes()->AttachAttribute(TangentUAttribName, new TDynamicMeshVertexAttribute<float, 3>(&Mesh));
+		Mesh.Attributes()->AttachAttribute(TangentVAttribName, new TDynamicMeshVertexAttribute<float, 1>(&Mesh)); // just a sign
+		TDynamicMeshScalarTriangleAttribute<bool>* VisAttrib = new TDynamicMeshScalarTriangleAttribute<bool>(&Mesh);
+		VisAttrib->Initialize(true);
+		Mesh.Attributes()->AttachAttribute(VisibleAttribName, VisAttrib);
+	}
+
+	// re-setup attributes to augment the mesh; only call this if we've lost the attached attributes
+	void Augment()
+	{
+		Mesh.EnableAttributes();
+		Mesh.Attributes()->EnableMaterialID();
+		ensure(Mesh.Attributes()->NumAttachedAttributes() == 0);
+		Mesh.Attributes()->AttachAttribute(TangentUAttribName, new TDynamicMeshVertexAttribute<float, 3>(&Mesh));
+		Mesh.Attributes()->AttachAttribute(TangentVAttribName, new TDynamicMeshVertexAttribute<float, 1>(&Mesh)); // just a sign
+		TDynamicMeshScalarTriangleAttribute<bool>* VisAttrib = new TDynamicMeshScalarTriangleAttribute<bool>(&Mesh);
+		VisAttrib->Initialize(true);
+		Mesh.Attributes()->AttachAttribute(VisibleAttribName, VisAttrib);
+	}
+
+	void SetDefaultMaterials(int32 MaterialID)
+	{
+		for (int TID : Mesh.TriangleIndicesItr())
+		{
+			Mesh.Attributes()->GetMaterialID()->SetNewValue(TID, MaterialID);
+		}
+	}
+
+	void SetDefaultAttributes(const FInternalSurfaceMaterials& Materials)
+	{
+		TDynamicMeshVertexAttribute<float, 3>* Us =
+			static_cast<TDynamicMeshVertexAttribute<float, 3>*>(Mesh.Attributes()->GetAttachedAttribute(TangentUAttribName));
+		TDynamicMeshVertexAttribute<float, 1>* VSigns =
+			static_cast<TDynamicMeshVertexAttribute<float, 1>*>(Mesh.Attributes()->GetAttachedAttribute(TangentVAttribName));
+
+		for (int VID : Mesh.VertexIndicesItr())
+		{
+			FVector3f N = Mesh.GetVertexNormal(VID);
+			FVector3f U;
+			VectorUtil::MakePerpVector(N, U);
+			Us->SetValue(VID, U);
+			float VSign = 1;
+			VSigns->SetValue(VID, &VSign);
+		}
+
+		TDynamicMeshScalarTriangleAttribute<bool>* Visible =
+			static_cast<TDynamicMeshScalarTriangleAttribute<bool>*>(Mesh.Attributes()->GetAttachedAttribute(VisibleAttribName));
+		for (int TID : Mesh.TriangleIndicesItr())
+		{
+			Visible->SetNewValue(TID, Materials.bGlobalVisibility);
+		}
+	}
+
+	void SetVisibility(int TID, bool bIsVisible)
+	{
+		TDynamicMeshScalarTriangleAttribute<bool>* Visible =
+			static_cast<TDynamicMeshScalarTriangleAttribute<bool>*>(Mesh.Attributes()->GetAttachedAttribute(VisibleAttribName));
+		Visible->SetValue(TID, bIsVisible);
+	}
+
+	bool GetVisibility(int TID)
+	{
+		TDynamicMeshScalarTriangleAttribute<bool>* Visible =
+			static_cast<TDynamicMeshScalarTriangleAttribute<bool>*>(Mesh.Attributes()->GetAttachedAttribute(VisibleAttribName));
+		return Visible->GetValue(TID);
+	}
+
+	void SetTangent(int VID, FVector3f Normal, FVector3f TangentU, FVector3f TangentV)
+	{
+		TDynamicMeshVertexAttribute<float, 3>* Us =
+			static_cast<TDynamicMeshVertexAttribute<float, 3>*>(Mesh.Attributes()->GetAttachedAttribute(TangentUAttribName));
+		TDynamicMeshVertexAttribute<float, 1>* VSigns =
+			static_cast<TDynamicMeshVertexAttribute<float, 1>*>(Mesh.Attributes()->GetAttachedAttribute(TangentVAttribName));
+		float VSign = FMathf::SignNonZero(Normal.Cross(TangentU).Dot(TangentV));
+		Us->SetValue(VID, TangentU);
+		VSigns->SetValue(VID, &VSign);
+	}
+	void FixTangents()
+	{
+		TDynamicMeshVertexAttribute<float, 3>* Us =
+			static_cast<TDynamicMeshVertexAttribute<float, 3>*>(Mesh.Attributes()->GetAttachedAttribute(TangentUAttribName));
+		TDynamicMeshVertexAttribute<float, 1>* VSigns =
+			static_cast<TDynamicMeshVertexAttribute<float, 1>*>(Mesh.Attributes()->GetAttachedAttribute(TangentVAttribName));
+
+		for (int VID : Mesh.VertexIndicesItr())
+		{
+			FVector3f N = Mesh.GetVertexNormal(VID);
+			FVector3f U;
+			Us->GetValue(VID, U);
+			U = N.Cross(U).Cross(N);
+			U.Normalize();
+			Us->SetValue(VID, U);
+			float VSign;
+			VSigns->GetValue(VID, &VSign);
+			VSign = FMathf::SignNonZero(VSign);
+			VSigns->SetValue(VID, &VSign);
+		}
+	}
+	void GetTangent(int VID, FVector3f& U, FVector3f& V)
+	{
+		TDynamicMeshVertexAttribute<float, 3>* Us =
+			static_cast<TDynamicMeshVertexAttribute<float, 3>*>(Mesh.Attributes()->GetAttachedAttribute(TangentUAttribName));
+		TDynamicMeshVertexAttribute<float, 1>* VSigns =
+			static_cast<TDynamicMeshVertexAttribute<float, 1>*>(Mesh.Attributes()->GetAttachedAttribute(TangentVAttribName));
+		FVector3f Normal = Mesh.GetVertexNormal(VID);
+		Us->GetValue(VID, U);
+		float VSign;
+		VSigns->GetValue(VID, &VSign);
+		V = Normal.Cross(U) * FMathf::SignNonZero(VSign);
+	}
+};
+
+struct FCellMeshes
+{
+	struct FCellInfo
+	{
+		FAugmentedDynamicMesh AugMesh;
+
+		// TODO: compute spatial in advance?  (only useful if we rework mesh booleans to support it)
+		//FDynamicMeshAABBTree3 Spatial;
+	};
+	TArray<FCellInfo> CellMeshes;
+	int32 OutsideCellIndex = -1;
+
+	// Noise Offsets, to randomize where perlin noise is sampled
+	FVector NoiseOffsetX;
+	FVector NoiseOffsetY;
+	FVector NoiseOffsetZ;
+
+	FCellMeshes() {}
+	FCellMeshes(const FPlanarCells& Cells, FAxisAlignedBox3d DomainBounds, double ExtendDomain, bool bIncludeOutsideCell = true)
+	{
+		Init(Cells, DomainBounds, ExtendDomain, bIncludeOutsideCell);
+	}
+
+	void RemeshForNoise(FDynamicMesh3& Mesh, EEdgeRefineFlags EdgeFlags, double TargetEdgeLen)
+	{
+		FQueueRemesher Remesh(&Mesh);
+		FMeshConstraints Constraints;
+
+		FMeshBoundaryLoops Boundary(&Mesh);
+		checkSlow(Boundary.GetLoopCount() == 1);
+
+		for (int VID : Mesh.VertexIndicesItr())
+		{
+			FVertexConstraint FullyConstrain(true, false, VID);
+			Constraints.SetOrUpdateVertexConstraint(VID, FullyConstrain);
+		}
+
+
+		FEdgeConstraint EdgeConstraint(EdgeFlags);
+		for (int EID : Boundary[0].Edges)
+		{
+			Constraints.SetOrUpdateEdgeConstraint(EID, EdgeConstraint);
+		}
+		Remesh.SetExternalConstraints(Constraints);
+		Remesh.SetTargetEdgeLength(TargetEdgeLen);
+		Remesh.Precompute();
+		Remesh.FastestRemesh();
+	}
+
+	float OctaveNoise(const FVector& V, const FNoiseSettings& Settings)
+	{
+		int32 Octaves = Settings.Octaves;
+		float NoiseValue = 0;
+		float OctaveScale = 1;
+		for (int32 Octave = 0; Octave < Octaves; Octave++, OctaveScale *= 2)
+		{
+			NoiseValue += FMath::PerlinNoise3D(V * OctaveScale) / OctaveScale;
+		}
+		return NoiseValue;
+	}
+
+	FVector NoiseVector(const FVector& Pos, const FNoiseSettings& Settings)
+	{
+		float Frequency = Settings.Frequency;
+		FVector Base = Pos * Frequency;
+		return FVector(
+			OctaveNoise(Base + NoiseOffsetX, Settings),
+			OctaveNoise(Base + NoiseOffsetY, Settings),
+			OctaveNoise(Base + NoiseOffsetZ, Settings)
+		) * Settings.Amplitude;
+	}
+
+	FVector3d NoiseDisplacement(const FVector3d& Pos, const FNoiseSettings& Settings)
+	{
+		FVector P = FVector(Pos);
+		return FVector3d(NoiseVector(P, Settings));
+	}
+
+	void ApplyNoise(FDynamicMesh3& Mesh, FVector3d Normal, const FNoiseSettings& Settings, bool bProjectBoundariesToNormal = false)
+	{
+		float Amplitude = Settings.Amplitude;
+		float Frequency = Settings.Frequency;
+		int32 Octaves = Settings.Octaves;
+		FVector3d Z = Normal * Amplitude;
+
+		for (int VID : Mesh.VertexIndicesItr())
+		{
+			FVector3d Pos = Mesh.GetVertex(VID);
+			FVector3d Displacement = NoiseDisplacement(Pos, Settings);
+			if (bProjectBoundariesToNormal || !Mesh.IsBoundaryVertex(VID))
+			{
+				// project displacement onto the normal direction
+				Displacement = Normal * Displacement.Dot(Normal);
+			}
+
+			
+			Mesh.SetVertex(VID, Pos + Displacement);
+		}
+	}
+
+	/**
+	 * Convert plane index to triangle group
+	 * @return group id encoding the source plane into a triangle mesh
+	 */
+	int PlaneToGroup(int Plane)
+	{
+		return Plane + 1;
+	}
+
+	/**
+	 * Convert group index to plane index
+	 * @return index of source plane for triangle, or -1 if no such plane
+	 */
+	int GroupToPlane(int Group)
+	{
+		return Group - 1;
+	}
+
+	void Init(const FPlanarCells& Cells, FAxisAlignedBox3d DomainBounds, double ExtendDomain, bool bIncludeOutsideCell = true)
+	{
+		NoiseOffsetX = FMath::VRand() * 100;
+		NoiseOffsetY = FMath::VRand() * 100;
+		NoiseOffsetZ = FMath::VRand() * 100;
+
+		float GlobalUVScale = Cells.InternalSurfaceMaterials.GlobalUVScale;
+		if (!ensure(GlobalUVScale > 0))
+		{
+			GlobalUVScale = 1;
+		}
+
+		int NumCells = Cells.NumCells;
+		bool bHasOutsideCell = false;
+		OutsideCellIndex = -1;
+		if (bIncludeOutsideCell && !Cells.IsInfinitePlane())
+		{
+			for (const TPair<int32, int32>& CellPair : Cells.PlaneCells)
+			{
+				if (CellPair.Value == -1)
+				{
+					bHasOutsideCell = true;
+					break;
+				}
+			}
+		}
+		if (bHasOutsideCell)
+		{
+			OutsideCellIndex = NumCells;
+			NumCells++;
+		}
+
+		CellMeshes.Reset();
+		CellMeshes.SetNum(NumCells);
+
+		DomainBounds.Expand(ExtendDomain);
+
+		bool bNoise = Cells.InternalSurfaceMaterials.NoiseSettings.IsSet();
+		if (bNoise)
+		{
+			ExtendDomain += Cells.InternalSurfaceMaterials.NoiseSettings->Amplitude;
+		}
+
+		// special handling for the infinite plane case; we need to adapt this to be a closed volume
+		if (Cells.IsInfinitePlane())
+		{
+			CreateMeshesForSinglePlane(Cells, DomainBounds, bNoise, GlobalUVScale);
+		}
+		else if (!bNoise) // bounded cells w/ no noise
+		{
+			CreateMeshesForBoundedPlanesWithoutNoise(NumCells, Cells, DomainBounds, bNoise, GlobalUVScale);
+		}
+		else // bounded cells with noise -- make each boundary plane separately so we can remesh them w/ noise vertices
+		{
+			CreateMeshesForBoundedPlanesWithNoise(NumCells, Cells, DomainBounds, bNoise, GlobalUVScale);
+		}
+		
+		// TODO: self-union on cells when it makes sense to do so (for non-single-plane inputs w/ high noise or possible untracked adjacencies)
+		/*for (FCellInfo& CellInfo : CellMeshes)
+		{
+			FMeshSelfUnion SelfUnion(&CellInfo.AugMesh.Mesh);
+			// TODO: need to have an option in SelfUnion to not weld edges
+			SelfUnion.Compute();
+		}*/
+
+		for (FCellInfo& CellInfo : CellMeshes)
+		{
+			check(CellInfo.AugMesh.Mesh.Attributes() != nullptr);
+			check(CellInfo.AugMesh.Mesh.Attributes()->NumAttachedAttributes() == 3);
+			CellInfo.AugMesh.SetDefaultAttributes(Cells.InternalSurfaceMaterials);
+			CellInfo.AugMesh.SetDefaultMaterials(Cells.InternalSurfaceMaterials.GlobalMaterialID);
+		}
+	}
+
+	void AppendMesh(FDynamicMesh3& Base, FDynamicMesh3& ToAppend, bool bFlipped)
+	{
+		FDynamicMeshEditor Editor(&Base);
+		FMeshIndexMappings Mapping;
+		Editor.AppendMesh(&ToAppend, Mapping);
+		if (bFlipped)
+		{
+			for (int TID : ToAppend.TriangleIndicesItr())
+			{
+				Base.ReverseTriOrientation(Mapping.GetNewTriangle(TID));
+			}
+			for (int VID : ToAppend.VertexIndicesItr())
+			{
+				int BaseVID = Mapping.GetNewVertex(VID);
+				Base.SetVertexNormal(BaseVID, -Base.GetVertexNormal(BaseVID));
+			}
+		}
+	}
+private:
+	void CreateMeshesForBoundedPlanesWithoutNoise(int NumCells, const FPlanarCells& Cells, const FAxisAlignedBox3d& DomainBounds, bool bNoise, double GlobalUVScale)
+	{
+		for (int32 PlaneIdx = 0; PlaneIdx < Cells.PlaneCells.Num(); PlaneIdx++)
+		{
+			const TPair<int32, int32>& CellPair = Cells.PlaneCells[PlaneIdx];
+			FDynamicMesh3* Meshes[2]{ &CellMeshes[CellPair.Key].AugMesh.Mesh, nullptr };
+			int32 OtherCell = CellPair.Value < 0 ? OutsideCellIndex : CellPair.Value;
+			int NumMeshes = OtherCell < 0 ? 1 : 2;
+			if (NumMeshes == 2)
+			{
+				Meshes[1] = &CellMeshes[OtherCell].AugMesh.Mesh;
+			}
+
+			const TArray<int>& PlaneBoundary = Cells.PlaneBoundaries[PlaneIdx];
+			FVector3f Normal(Cells.Planes[PlaneIdx].GetNormal());
+			FFrame3d PlaneFrame(Cells.Planes[PlaneIdx]);
+			FVertexInfo PlaneVertInfo;
+			PlaneVertInfo.bHaveC = true;
+			PlaneVertInfo.bHaveUV = true;
+			PlaneVertInfo.bHaveN = true;
+			PlaneVertInfo.Color = FVector3f(1, 1, 1);
+			int VertStart[2]{ -1, -1 };
+			for (int MeshIdx = 0; MeshIdx < NumMeshes; MeshIdx++)
+			{
+				PlaneVertInfo.Normal = Normal;
+				if (MeshIdx == 1 && OtherCell != OutsideCellIndex)
+				{
+					PlaneVertInfo.Normal *= -1;
+				}
+				VertStart[MeshIdx] = Meshes[MeshIdx]->MaxVertexID();
+				FVector2f MinUV(FMathf::MaxReal, FMathf::MaxReal);
+				for (int BoundaryVertex : PlaneBoundary)
+				{
+					FVector3d Position = FVector3d(Cells.PlaneBoundaryVertices[BoundaryVertex]);
+					FVector2f UV = FVector2f(PlaneFrame.ToPlaneUV(Position));
+					MinUV.X = FMathf::Min(PlaneVertInfo.UV.X, MinUV.X);
+					MinUV.Y = FMathf::Min(PlaneVertInfo.UV.Y, MinUV.Y);
+				}
+				for (int BoundaryVertex : PlaneBoundary)
+				{
+					PlaneVertInfo.Position = FVector3d(Cells.PlaneBoundaryVertices[BoundaryVertex]);
+					PlaneVertInfo.UV = (FVector2f(PlaneFrame.ToPlaneUV(PlaneVertInfo.Position)) - MinUV) * GlobalUVScale;
+					Meshes[MeshIdx]->AppendVertex(PlaneVertInfo);
+				}
+			}
+
+			for (int V0 = 0, V1 = 1, V2 = 2; V2 < PlaneBoundary.Num(); V1 = V2++)
+			{
+				for (int MeshIdx = 0; MeshIdx < NumMeshes; MeshIdx++)
+				{
+					int Offset = VertStart[MeshIdx];
+					FIndex3i Tri(V0 + Offset, V1 + Offset, V2 + Offset);
+					if (MeshIdx == 1 && OtherCell != OutsideCellIndex)
+					{
+						Swap(Tri.B, Tri.C);
+					}
+					Meshes[MeshIdx]->AppendTriangle(Tri, PlaneToGroup(PlaneIdx));
+				}
+			}
+		}
+	}
+	void CreateMeshesForBoundedPlanesWithNoise(int NumCells, const FPlanarCells& Cells, const FAxisAlignedBox3d& DomainBounds, bool bNoise, double GlobalUVScale)
+	{
+		TArray<FDynamicMesh3> PlaneMeshes;
+		PlaneMeshes.SetNum(Cells.Planes.Num());
+
+		struct FPlaneIdxAndFlip
+		{
+			int32 PlaneIdx;
+			bool bIsFlipped;
+		};
+		TArray<TArray<FPlaneIdxAndFlip>> CellPlanes; // per cell, the planes that border that cell
+		CellPlanes.SetNum(NumCells);
+
+		for (int32 PlaneIdx = 0; PlaneIdx < Cells.PlaneCells.Num(); PlaneIdx++)
+		{
+			const TPair<int32, int32>& CellPair = Cells.PlaneCells[PlaneIdx];
+			int32 OtherCell = CellPair.Value < 0 ? OutsideCellIndex : CellPair.Value;
+			if (ensure(CellPlanes.IsValidIndex(CellPair.Key)))
+			{
+				CellPlanes[CellPair.Key].Add({ PlaneIdx, false });
+			}
+			if (CellPlanes.IsValidIndex(OtherCell))
+			{
+				CellPlanes[OtherCell].Add({ PlaneIdx, true });
+			}
+		}
+
+		// heuristic to protect against creating too many vertices on remeshing
+		float MaxArea = 0;
+		for (int32 PlaneIdx = 0; PlaneIdx < Cells.Planes.Num(); PlaneIdx++)
+		{
+			const TArray<int>& PlaneBoundary = Cells.PlaneBoundaries[PlaneIdx];
+			const FVector& V0 = Cells.PlaneBoundaryVertices[PlaneBoundary[0]];
+			FVector AreaVec = FVector::ZeroVector;
+			for (int32 V1Idx = 1, V2Idx = 2; V2Idx < PlaneBoundary.Num(); V1Idx = V2Idx++)
+			{
+				const FVector& V1 = Cells.PlaneBoundaryVertices[PlaneBoundary[V1Idx]];
+				const FVector& V2 = Cells.PlaneBoundaryVertices[PlaneBoundary[V2Idx]];
+				AreaVec += (V1 - V0) ^ (V2 - V1);
+			}
+			MaxArea = AreaVec.Size();
+		}
+		double MaxVerts = 10000;
+		double MinEdgeLen = FMathd::Sqrt((double)MaxArea / MaxVerts);
+		double Spacing = FMath::Max3(.001, MinEdgeLen, (double)Cells.InternalSurfaceMaterials.NoiseSettings->PointSpacing);
+
+		ParallelFor(Cells.Planes.Num(), [this, &PlaneMeshes, &Cells, GlobalUVScale, Spacing](int32 PlaneIdx)
+			{
+				FDynamicMesh3& Mesh = PlaneMeshes[PlaneIdx];
+				const TArray<int>& PlaneBoundary = Cells.PlaneBoundaries[PlaneIdx];
+				FVector3f Normal(Cells.Planes[PlaneIdx].GetNormal());
+				FFrame3d PlaneFrame(Cells.Planes[PlaneIdx]);
+				FVertexInfo PlaneVertInfo;
+				PlaneVertInfo.bHaveC = true;
+				PlaneVertInfo.bHaveUV = true;
+				PlaneVertInfo.bHaveN = true;
+				PlaneVertInfo.Normal = Normal;
+				PlaneVertInfo.Color = FVector3f(1, 1, 1);
+
+				FVector2f MinUV(FMathf::MaxReal, FMathf::MaxReal);
+				for (int BoundaryVertex : PlaneBoundary)
+				{
+					FVector3d Position = FVector3d(Cells.PlaneBoundaryVertices[BoundaryVertex]);
+					FVector2f UV = FVector2f(PlaneFrame.ToPlaneUV(Position));
+					MinUV.X = FMathf::Min(PlaneVertInfo.UV.X, MinUV.X);
+					MinUV.Y = FMathf::Min(PlaneVertInfo.UV.Y, MinUV.Y);
+				}
+
+				FPolygon2f Polygon;
+				for (int BoundaryVertex : PlaneBoundary)
+				{
+					PlaneVertInfo.Position = FVector3d(Cells.PlaneBoundaryVertices[BoundaryVertex]);
+					PlaneVertInfo.UV = (FVector2f(PlaneFrame.ToPlaneUV(PlaneVertInfo.Position)) - MinUV) * GlobalUVScale;
+					Polygon.AppendVertex(PlaneVertInfo.UV);
+					Mesh.AppendVertex(PlaneVertInfo);
+				}
+
+				// we do a CDT here to give a slightly better start to remeshing; we could try simple ear clipping instead
+				FGeneralPolygon2f GeneralPolygon(Polygon);
+				FConstrainedDelaunay2f Triangulation;
+				Triangulation.FillRule = FConstrainedDelaunay2f::EFillRule::NonZero;
+				Triangulation.Add(GeneralPolygon);
+				Triangulation.Triangulate();
+
+				int GID = PlaneToGroup(PlaneIdx);
+				for (FIndex3i Triangle : Triangulation.Triangles)
+				{
+					Mesh.AppendTriangle(Triangle, GID);
+				}
+
+				RemeshForNoise(Mesh, EEdgeRefineFlags::SplitsOnly, Spacing);
+				ApplyNoise(Mesh, FVector3d(Normal), Cells.InternalSurfaceMaterials.NoiseSettings.GetValue());
+				FMeshNormals::QuickComputeVertexNormals(Mesh);
+			}, EParallelForFlags::None);
+
+		for (int CellIdx = 0; CellIdx < NumCells; CellIdx++)
+		{
+			FCellInfo& CellInfo = CellMeshes[CellIdx];
+			FDynamicMesh3& Mesh = CellInfo.AugMesh.Mesh;
+			bool bFlipForOutsideCell = CellIdx == OutsideCellIndex; // outside cell will be subtracted, and needs all planes flipped vs normal
+			for (FPlaneIdxAndFlip PlaneInfo : CellPlanes[CellIdx])
+			{
+				AppendMesh(Mesh, PlaneMeshes[PlaneInfo.PlaneIdx], PlaneInfo.bIsFlipped ^ bFlipForOutsideCell);
+			}
+		}
+	}
+	void CreateMeshesForSinglePlane(const FPlanarCells& Cells, const FAxisAlignedBox3d& DomainBounds, bool bNoise, double GlobalUVScale)
+	{
+		int GID = PlaneToGroup(0);
+		FPlane Plane = Cells.Planes[0];
+
+		FFrame3d PlaneFrame(Plane);
+		FInterval1d ZRange;
+		FAxisAlignedBox2d XYRange;
+		for (int CornerIdx = 0; CornerIdx < 8; CornerIdx++)
+		{
+			FVector3d Corner = DomainBounds.GetCorner(CornerIdx);
+			XYRange.Contain(PlaneFrame.ToPlaneUV(Corner));
+			ZRange.Contain(Plane.PlaneDot(FVector(Corner)));
+		}
+		//if (FMathd::SignNonZero(ZRange.Min) == FMathd::SignNonZero(ZRange.Max))
+		//{
+		//	// TODO: early out for plane that doesn't even intersect the domain bounding box?
+		//}
+
+		FDynamicMesh3 PlaneMesh;
+		FVertexInfo PlaneVertInfo;
+		PlaneVertInfo.bHaveC = true;
+		PlaneVertInfo.bHaveUV = true;
+		PlaneVertInfo.bHaveN = true;
+		PlaneVertInfo.Color = FVector3f(1, 1, 1);
+		PlaneVertInfo.Normal = FVector3f(Plane.GetNormal());
+
+		for (int CornerIdx = 0; CornerIdx < 4; CornerIdx++)
+		{
+			PlaneVertInfo.Position = PlaneFrame.FromPlaneUV(XYRange.GetCorner(CornerIdx));
+			PlaneVertInfo.UV = FVector2f(XYRange.GetCorner(CornerIdx) - XYRange.Min) * GlobalUVScale;
+			PlaneMesh.AppendVertex(PlaneVertInfo);
+		}
+		PlaneMesh.AppendTriangle(0, 1, 2, GID);
+		PlaneMesh.AppendTriangle(0, 2, 3, GID);
+
+		if (bNoise)
+		{
+			RemeshForNoise(PlaneMesh, EEdgeRefineFlags::SplitsOnly, Cells.InternalSurfaceMaterials.NoiseSettings->PointSpacing);
+			ApplyNoise(PlaneMesh, PlaneFrame.GetAxis(2), Cells.InternalSurfaceMaterials.NoiseSettings.GetValue(), true);
+			FMeshNormals::QuickComputeVertexNormals(PlaneMesh);
+		}
+		FDynamicMesh3* Meshes[2]{ &CellMeshes[0].AugMesh.Mesh, &CellMeshes[1].AugMesh.Mesh };
+		for (int Side = 0; Side < 2; Side++)
+		{
+			*Meshes[Side] = PlaneMesh;
+			double Offset = ZRange.Max;
+			FMeshBoundaryLoops Boundary(Meshes[Side]);
+			checkSlow(Boundary.GetLoopCount() == 1);
+			int FirstIdx;
+			bool bFound = Boundary[0].Vertices.Find(0, FirstIdx);
+			checkSlow(bFound);
+			TArray<int> VertIDs[2], MatchedIndices[2];
+			VertIDs[0] = Boundary[0].Vertices;
+			if (FirstIdx != 0)
+			{
+				Algo::Rotate(VertIDs[0], FirstIdx);
+			}
+			checkSlow(VertIDs[0][0] == 0);
+			MatchedIndices[0].Add(0);
+			int FoundIndices = 1;
+			for (int VIDIdx = 0; VIDIdx < VertIDs[0].Num(); VIDIdx++)
+			{
+				int VID = VertIDs[0][VIDIdx];
+				if (VID == FoundIndices)
+				{
+					FoundIndices++;
+					MatchedIndices[0].Add(VIDIdx);
+				}
+			}
+
+			if (Side == 0)
+			{
+				Meshes[Side]->ReverseOrientation(true);
+				Offset = ZRange.Min;
+			}
+			PlaneVertInfo.Normal = FVector3f(Plane.GetNormal()) * (-1 + Side * 2);
+			FVector3d OffsetVec = FVector3d(Plane.GetNormal()) * Offset;
+
+			for (int CornerIdx = 0; CornerIdx < 4; CornerIdx++)
+			{
+				PlaneVertInfo.Position = Meshes[Side]->GetVertex(CornerIdx) + OffsetVec;
+				// UVs shouldn't matter for outer box vertices because they're outside of the domain by construction ...
+				VertIDs[1].Add(Meshes[Side]->AppendVertex(PlaneVertInfo));
+				MatchedIndices[1].Add(CornerIdx);
+			}
+			int NewTris[2]{
+				Meshes[Side]->AppendTriangle(VertIDs[1][0], VertIDs[1][1], VertIDs[1][2], GID),
+				Meshes[Side]->AppendTriangle(VertIDs[1][0], VertIDs[1][2], VertIDs[1][3], GID)
+			};
+			if (Side == 1)
+			{
+				Meshes[Side]->ReverseTriOrientation(NewTris[0]);
+				Meshes[Side]->ReverseTriOrientation(NewTris[1]);
+			}
+			FDynamicMeshEditor Editor(Meshes[Side]);
+			FDynamicMeshEditResult ResultOut;
+			Editor.StitchSparselyCorrespondedVertexLoops(VertIDs[0], MatchedIndices[0], VertIDs[1], MatchedIndices[1], ResultOut, Side == 0);
+
+			// re-enable tangents and visibility attributes, since these are lost when we set the mesh to a copy of the plane mesh
+			CellMeshes[Side].AugMesh.Augment();
+		}
+	}
+};
+
+
+
+// organize metadata corresponding dynamic mesh and geometry collection data
+struct FDynamicMeshCollection
+{
+	struct FMeshData
+	{
+		FAugmentedDynamicMesh AugMesh;
+		
+		// FDynamicMeshAABBTree3 Spatial; // TODO: maybe refactor mesh booleans to allow version where caller provides spatial data; it's computed every boolean now
+		// FTransform3d Transform; // TODO: maybe pretransform the data to a space that is good for cutting; refactor mesh boolean so there is an option to have it not transform input
+		int32 TransformIndex; // where the mesh was from in the geometry collection
+		FTransform ToCollection; // transform that need be applied to go back to the local space of the geometry collection
+	};
+	TArray<FMeshData> Meshes;
+	FAxisAlignedBox3d Bounds;
+
+	FDynamicMeshCollection(const FGeometryCollection* Collection, const TArrayView<const int32>& TransformIndices, FTransform TransformCells)
+	{
+		Init(Collection, TransformIndices, TransformCells);
+	}
+
+	void Init(const FGeometryCollection* Collection, const TArrayView<const int32>& TransformIndices, FTransform TransformCells)
+	{
+		FTransform WorldToCells = TransformCells.Inverse();
+
+		Meshes.Reset();
+		Bounds = FAxisAlignedBox3d::Empty();
+
+		for (int32 TransformIdx : TransformIndices)
+		{
+			if (Collection->Children[TransformIdx].Num() > 0)
+			{
+				// only store the meshes of leaf nodes
+				continue;
+			}
+
+			FTransform3d CollectionToLocal = FTransform3d(GeometryCollectionAlgo::GlobalMatrix(Collection->Transform, Collection->Parent, TransformIdx) * WorldToCells);
+
+			FMeshData& MeshData = Meshes.Emplace_GetRef();
+			MeshData.TransformIndex = TransformIdx;
+			MeshData.ToCollection = FTransform(CollectionToLocal.Inverse());
+			FDynamicMesh3& Mesh = MeshData.AugMesh.Mesh;
+
+			int32 GeometryIdx = Collection->TransformToGeometryIndex[TransformIdx];
+			Mesh.EnableAttributes();
+			Mesh.Attributes()->EnableMaterialID();
+
+			int32 VertexStart = Collection->VertexStart[GeometryIdx];
+			int32 VertexCount = Collection->VertexCount[GeometryIdx];
+			int32 FaceCount = Collection->FaceCount[GeometryIdx];
+
+			FVertexInfo VertexInfo;
+			VertexInfo.bHaveC = true;
+			VertexInfo.bHaveN = true;
+			VertexInfo.bHaveUV = true;
+			for (int32 Idx = VertexStart, N = VertexStart + VertexCount; Idx < N; Idx++)
+			{
+				VertexInfo.Position = CollectionToLocal.TransformPosition(FVector3d(Collection->Vertex[Idx]));
+				VertexInfo.UV = FVector2f(Collection->UV[Idx]);
+				VertexInfo.Color = FVector3f(Collection->Color[Idx]);
+				VertexInfo.Normal = CollectionToLocal.TransformNormal(FVector3f(Collection->Normal[Idx]));
+				int VID = Mesh.AppendVertex(VertexInfo);
+				MeshData.AugMesh.SetTangent(VID, VertexInfo.Normal, Collection->TangentU[Idx], Collection->TangentV[Idx]);
+			}
+			FIntVector VertexOffset(VertexStart, VertexStart, VertexStart);
+			for (int32 Idx = Collection->FaceStart[GeometryIdx], N = Collection->FaceStart[GeometryIdx] + FaceCount; Idx < N; Idx++)
+			{
+				int TID = Mesh.AppendTriangle(FIndex3i(Collection->Indices[Idx] - VertexOffset), 0);
+				Mesh.Attributes()->GetMaterialID()->SetValue(TID, Collection->MaterialID[Idx]);
+				MeshData.AugMesh.SetVisibility(TID, Collection->Visible[Idx]);
+				// note: material index doesn't need to be passed through; will be rebuilt by a call to reindex materials once the cut mesh is returned back to geometry collection format
+			}
+
+			FDynamicMeshEditor Editor(&Mesh);
+			Editor.RemoveIsolatedVertices();
+
+			Bounds.Contain(Mesh.GetCachedBounds());
+
+			// TODO: build spatial data (add this after setting up mesh boolean path that can use it)
+			//MeshData.Spatial.SetMesh(&Mesh);
+		}
+	}
+
+
+	int32 CutWithMultiplePlanes(
+		const TArrayView<const FPlane>& Planes, 
+		FGeometryCollection* Collection,
+		FInternalSurfaceMaterials& InternalSurfaceMaterials,
+		bool bSetDefaultInternalMaterialsFromCollection
+		)
+	{
+#if WITH_EDITOR
+		// Create progress indicator dialog
+		static const FText SlowTaskText = NSLOCTEXT("CutMultipleWithMultiplePlanes", "CutMultipleWithMultiplePlanesText", "Cutting geometry collection with plane(s)...");
+
+		FScopedSlowTask SlowTask(Planes.Num(), SlowTaskText);
+		SlowTask.MakeDialog();
+
+		// Declare progress shortcut lambdas
+		auto EnterProgressFrame = [&SlowTask](float Progress)
+		{
+			SlowTask.EnterProgressFrame(Progress);
+		};
+#else
+		auto EnterProgressFrame = [](float Progress) {};
+#endif
+
+		bool bHasProximity = Collection->HasAttribute("Proximity", FGeometryCollection::GeometryGroup);
+		TArray<TUniquePtr<FMeshData>> ToCut;
+		TArray<TUniquePtr<TPointHashGrid3d<int>>> VerticesHashes;
+		auto HashMeshVertices = [&VerticesHashes, &ToCut](int32 HashIdx)
+		{
+			FDynamicMesh3& Mesh = ToCut[HashIdx]->AugMesh.Mesh;
+			if (HashIdx >= VerticesHashes.Num())
+			{
+				VerticesHashes.SetNum(HashIdx + 1);
+			}
+			if (VerticesHashes[HashIdx].IsValid())
+			{
+				return;
+			}
+			VerticesHashes[HashIdx] = MakeUnique<TPointHashGrid3d<int>>(FMathd::ZeroTolerance * 1000, -1);
+			TPointHashGrid3d<int>& Grid = *VerticesHashes[HashIdx].Get();
+			for (int VID : Mesh.VertexIndicesItr())
+			{
+				Grid.InsertPointUnsafe(VID, Mesh.GetVertex(VID));
+			}
+		};
+		auto ClearHash = [&VerticesHashes](int32 HashIdx)
+		{
+			if (HashIdx < VerticesHashes.Num())
+			{
+				VerticesHashes[HashIdx].Release();
+			}
+		};
+		auto IsNeighbor = [&ToCut, &VerticesHashes](int32 A, int32 B)
+		{
+			if (!ensure(A < ToCut.Num() && B < ToCut.Num() && A < VerticesHashes.Num() && B < VerticesHashes.Num()))
+			{
+				return false;
+			}
+			if (!ensure(VerticesHashes[A].IsValid() && VerticesHashes[B].IsValid()))
+			{
+				return false;
+			}
+			if (!ToCut[A]->AugMesh.Mesh.GetCachedBounds().Intersects(ToCut[B]->AugMesh.Mesh.GetCachedBounds()))
+			{
+				return false;
+			}
+			if (ToCut[A]->AugMesh.Mesh.VertexCount() > ToCut[B]->AugMesh.Mesh.VertexCount())
+			{
+				Swap(A, B);
+			}
+			FDynamicMesh3& RefMesh = ToCut[B]->AugMesh.Mesh;
+			for (const FVector3d& V : ToCut[A]->AugMesh.Mesh.VerticesItr())
+			{
+				TPair<int, double> Nearest = VerticesHashes[B]->FindNearestInRadius(V, FMathd::ZeroTolerance * 10, [&RefMesh, &V](int VID)
+					{
+						return RefMesh.GetVertex(VID).DistanceSquared(V);
+					});
+				if (Nearest.Key != -1)
+				{
+					return true;
+				}
+			}
+			return false;
+		};
+		// copy initial surfaces
+		for (FMeshData& MeshData : Meshes)
+		{
+			ToCut.Add(MakeUnique<FMeshData>(MeshData));
+		}
+		// track connections between meshes via their indices in the ToCut array
+		TMultiMap<int32, int32> Proximity;
+		auto ProxLink = [&Proximity](int32 A, int32 B)
+		{
+			Proximity.Add(A, B);
+			Proximity.Add(B, A);
+		};
+		auto ProxUnlink = [&Proximity](int32 A, int32 B)
+		{
+			Proximity.RemoveSingle(A, B);
+			Proximity.RemoveSingle(B, A);
+		};
+		for (int32 PlaneIdx = 0; PlaneIdx < Planes.Num(); PlaneIdx++)
+		{
+			EnterProgressFrame(1);
+			FPlanarCells PlaneCells(Planes[PlaneIdx]);
+			PlaneCells.InternalSurfaceMaterials = InternalSurfaceMaterials;
+			FCellMeshes CellMeshes(PlaneCells, Bounds, InternalSurfaceMaterials.NoiseSettings.IsSet() ? InternalSurfaceMaterials.NoiseSettings->Amplitude : 0, false);
+
+			// TODO: we could do these cuts in parallel (will takes some rework of the proximity and how results are added to the ToCut array)
+			for (int32 ToCutIdx = 0, ToCutNum = ToCut.Num(); ToCutIdx < ToCutNum; ToCutIdx++)
+			{
+				FMeshData& Surface = *ToCut[ToCutIdx];
+
+				FAxisAlignedBox3d Box = Surface.AugMesh.Mesh.GetCachedBounds();
+				if (InternalSurfaceMaterials.NoiseSettings.IsSet())
+				{
+					Box.Expand(InternalSurfaceMaterials.NoiseSettings->Amplitude);
+				}
+				if (!FMath::PlaneAABBIntersection(Planes[PlaneIdx], FBox(Box)))
+				{
+					continue;
+				}
+
+				TArray<TUniquePtr<FMeshData>> BoolResults;
+				for (int ResultIdx = 0; ResultIdx < 2; ResultIdx++)
+				{
+					BoolResults.Add(MakeUnique<FMeshData>());
+					BoolResults[ResultIdx]->TransformIndex = Surface.TransformIndex;
+					BoolResults[ResultIdx]->ToCollection = Surface.ToCollection;
+				}
+				check(CellMeshes.CellMeshes.Num() == 2);
+				bool bKeepResults = true;
+				for (int32 CellIdx = 0; CellIdx < 2; CellIdx++)
+				{
+					FCellMeshes::FCellInfo& Cell = CellMeshes.CellMeshes[CellIdx];
+
+					FMeshBoolean::EBooleanOp Op = FMeshBoolean::EBooleanOp::Intersect;
+					if (CellIdx == CellMeshes.OutsideCellIndex)
+					{
+						Op = FMeshBoolean::EBooleanOp::Difference;
+					}
+					FMeshBoolean Boolean(&Surface.AugMesh.Mesh, &Cell.AugMesh.Mesh, &BoolResults[CellIdx]->AugMesh.Mesh, Op);
+					Boolean.bWeldSharedEdges = false;
+					if (!Boolean.Compute())
+					{
+						// TODO: do something about failure cases?  e.g. try auto-filling small holes?
+						// note: failure cases won't be detected at all unless we weld edges,
+						//       which will require re-working how tangents are carried through
+					}
+					if (BoolResults[CellIdx]->AugMesh.Mesh.TriangleCount() == 0)
+					{
+						bKeepResults = false;
+						break;
+					}
+				}
+
+				if (bKeepResults)
+				{
+					ToCut[ToCutIdx] = MoveTemp(BoolResults[0]);
+					int32 NewIdx = ToCut.Add(MoveTemp(BoolResults[1]));
+					int32 ResultIndices[2] = { ToCutIdx, NewIdx };
+
+					// update proximity for neighbors of the original piece
+					if (bHasProximity)
+					{
+						ClearHash(ToCutIdx);
+						TArray<int32> Nbrs;
+						Proximity.MultiFind(ToCutIdx, Nbrs);
+						if (Nbrs.Num() > 0)
+						{
+							HashMeshVertices(ToCutIdx);
+							HashMeshVertices(NewIdx);
+
+							for (int32 Nbr : Nbrs)
+							{
+								ProxUnlink(ToCutIdx, Nbr);
+								HashMeshVertices(Nbr);
+								for (int32 Idx = 0; Idx < 2; Idx++)
+								{
+									int32 ResultIdx = ResultIndices[Idx];
+									int32 OldIdx = Nbr;
+									if (IsNeighbor(ResultIdx, OldIdx))
+									{
+										ProxLink(ResultIdx, OldIdx);
+									}
+								}
+							}
+						}
+					}
+
+					// add the connection between the two new pieces
+					ProxLink(ResultIndices[0], ResultIndices[1]);
+				}
+			}
+		}
+
+		TMultiMap<int32, int32> ParentTransformToChildren;
+		for (int32 ToCutIdx = 0; ToCutIdx < ToCut.Num(); ToCutIdx++)
+		{
+			ParentTransformToChildren.Add(ToCut[ToCutIdx]->TransformIndex, ToCutIdx);
+		}
+
+		TArray<int32> ToCutIdxToGeometryIdx;  ToCutIdxToGeometryIdx.Init(-1, ToCut.Num());
+		TArray<int32> ToCutIndices;
+		int32 FirstCreatedIndex = -1;
+		for (FMeshData& MeshData : Meshes)
+		{
+			int32 GeometryIdx = Collection->TransformToGeometryIndex[MeshData.TransformIndex];
+			int32 MaterialIDOverride = bSetDefaultInternalMaterialsFromCollection ? InternalSurfaceMaterials.GetDefaultMaterialIDForGeometry(*Collection, GeometryIdx) : -1;
+			ParentTransformToChildren.MultiFind(MeshData.TransformIndex, ToCutIndices);
+
+			// if there's only one mesh here, i.e. it didn't get cut at all
+			if (ToCutIndices.Num() <= 1)
+			{
+				continue;
+			}
+
+			// hide old parent geometry
+			SetVisibility(*Collection, GeometryIdx, false);
+
+			// add newly created geometry as children
+			int32 SubPartIdx = 0;
+			for (int32 ToCutIdx : ToCutIndices)
+			{
+				ToCut[ToCutIdx]->AugMesh.FixTangents();
+				for (int TID : ToCut[ToCutIdx]->AugMesh.Mesh.TriangleIndicesItr())
+				{
+					int GID = ToCut[ToCutIdx]->AugMesh.Mesh.GetTriangleGroup(TID);
+					if (GID > 0)
+					{
+						if (MaterialIDOverride != -1)
+						{
+							ToCut[ToCutIdx]->AugMesh.Mesh.Attributes()->GetMaterialID()->SetNewValue(TID, MaterialIDOverride);
+						}
+					}
+				}
+				int32 CreatedGeometryIdx = AppendToCollection(ToCut[ToCutIdx]->ToCollection, ToCut[ToCutIdx]->AugMesh, ToCut[ToCutIdx]->TransformIndex, SubPartIdx++, *Collection);
+				ToCutIdxToGeometryIdx[ToCutIdx] = CreatedGeometryIdx;
+				if (FirstCreatedIndex == -1)
+				{
+					FirstCreatedIndex = CreatedGeometryIdx;
+				}
+			}
+		}
+
+		// create proximity sets on geometry collection and populate using ToCut's Proximity multimap and the array ToCutIdxToGeometryIdx
+		if (bHasProximity)
+		{
+			TManagedArray<TSet<int32>>& GCProximity = Collection->GetAttribute<TSet<int32>>("Proximity", FGeometryCollection::GeometryGroup);
+			for (TPair<int32, int32> Link : Proximity)
+			{
+				GCProximity[ToCutIdxToGeometryIdx[Link.Key]].Add(ToCutIdxToGeometryIdx[Link.Value]);
+			}
+		}
+
+		return FirstCreatedIndex;
+	}
+
+	/**
+	 * Cut collection meshes with cell meshes, and append results to a geometry collection
+	 *
+	 * @param PlanarCells The source definitions of the cells to cut with
+	 * @param CellsMeshes Meshed versions of the cells, with noise and material properties baked in
+	 * @param Collection Results will be stored in this
+	 * @param bSetDefaultInternalMaterialsFromCollection If true, set internal materials to the most common external material + 1, following a convenient artist convention
+	 * @return Index of the first created geometry
+	 */
+	int32 CutWithCellMeshes(const FPlanarCells& PlanarCells, FCellMeshes& CellMeshes, FGeometryCollection* Collection, bool bSetDefaultInternalMaterialsFromCollection)
+	{
+		// TODO: should we do these cuts in parallel, and the appends sequentially below?
+		int32 FirstIdx = -1;
+		int BadCount = 0;
+		bool bHasProximity = Collection->HasAttribute("Proximity", FGeometryCollection::GeometryGroup);
+		for (FMeshData& Surface : Meshes)
+		{
+			int32 GeometryIdx = Collection->TransformToGeometryIndex[Surface.TransformIndex];
+			TArray<FAugmentedDynamicMesh> BooleanResults;  BooleanResults.SetNum(CellMeshes.CellMeshes.Num());
+			ParallelFor(CellMeshes.CellMeshes.Num(), [&BooleanResults, &CellMeshes, &Surface](int32 CellIdx)
+				{
+					FCellMeshes::FCellInfo& Cell = CellMeshes.CellMeshes[CellIdx];
+					if (Cell.AugMesh.Mesh.GetCachedBounds().Intersects(Surface.AugMesh.Mesh.GetCachedBounds()))
+					{
+						FAugmentedDynamicMesh& AugBoolResult = BooleanResults[CellIdx];
+
+						FMeshBoolean::EBooleanOp Op = FMeshBoolean::EBooleanOp::Intersect;
+						if (CellIdx == CellMeshes.OutsideCellIndex)
+						{
+							Op = FMeshBoolean::EBooleanOp::Difference;
+						}
+						FMeshBoolean Boolean(&Surface.AugMesh.Mesh, &Cell.AugMesh.Mesh, &AugBoolResult.Mesh, Op);
+						Boolean.bWeldSharedEdges = false;
+						if (!Boolean.Compute())
+						{
+							// TODO: do something about failure cases?  e.g. try auto-filling small holes?
+							// note: failure cases won't be detected at all unless we weld edges,
+							//       which will require re-working how tangents are carried through
+						}
+					}
+				}, EParallelForFlags::None);
+
+			int32 NonEmptyResults = 0;
+			for (const FAugmentedDynamicMesh& AugBoolResult : BooleanResults)
+			{
+				if (AugBoolResult.Mesh.TriangleCount() > 0)
+				{
+					NonEmptyResults++;
+				}
+			}
+
+			if (NonEmptyResults > 1) // only write to geometry collection if more than one result was non-empty
+			{
+				TSet<int32> PlanesInOutput;
+				TMap<int32, int32> CellToGeometry;
+				int32 SubPartIndex = 0;
+				int32 MaterialIDOverride = bSetDefaultInternalMaterialsFromCollection ? PlanarCells.InternalSurfaceMaterials.GetDefaultMaterialIDForGeometry(*Collection, GeometryIdx) : -1;
+
+				for (int32 CellIdx = 0; CellIdx < CellMeshes.CellMeshes.Num(); CellIdx++)
+				{
+					FAugmentedDynamicMesh& AugBoolResult = BooleanResults[CellIdx];
+					if (AugBoolResult.Mesh.TriangleCount() > 0)
+					{
+						AugBoolResult.FixTangents();
+						for (int TID : AugBoolResult.Mesh.TriangleIndicesItr())
+						{
+							int GID = AugBoolResult.Mesh.GetTriangleGroup(TID);
+							int32 PlaneIdx = CellMeshes.GroupToPlane(GID);
+							if (PlaneIdx >= 0)
+							{
+								PlanesInOutput.Add(PlaneIdx);
+								if (MaterialIDOverride != -1)
+								{
+									AugBoolResult.Mesh.Attributes()->GetMaterialID()->SetNewValue(TID, MaterialIDOverride);
+								}
+							}
+						}
+						int32 CreatedGeometryIdx = AppendToCollection(Surface.ToCollection, AugBoolResult, Surface.TransformIndex, SubPartIndex, *Collection);
+						CellToGeometry.Add(CellIdx, CreatedGeometryIdx);
+						if (FirstIdx == -1)
+						{
+							FirstIdx = CreatedGeometryIdx;
+						}
+						SubPartIndex++;
+					}
+					AugBoolResult.Mesh.Clear(); // we can stop storing the dynamic mesh version now that it's in the geometry collection
+				}
+				if (bHasProximity)
+				{
+					TManagedArray<TSet<int32>>& Proximity = Collection->GetAttribute<TSet<int32>>("Proximity", FGeometryCollection::GeometryGroup);
+					for (int32 PlaneIdx : PlanesInOutput)
+					{
+						TPair<int32, int32> Cells = PlanarCells.PlaneCells[PlaneIdx];
+						int32 SecondCell = Cells.Value == -1 ? CellMeshes.OutsideCellIndex : Cells.Value;
+						if (SecondCell != -1)
+						{
+							int32* GeomA = CellToGeometry.Find(Cells.Key);
+							int32* GeomB = CellToGeometry.Find(SecondCell);
+							if (GeomA && GeomB)
+							{
+								Proximity[*GeomA].Add(*GeomB);
+								Proximity[*GeomB].Add(*GeomA);
+							}
+						}
+					}
+				}
+				// turn off old geom visibility (preferred default behavior)
+				SetVisibility(*Collection, GeometryIdx, false);
+			}
+		}
+
+		return FirstIdx;
+	}
+
+	static void SetVisibility(FGeometryCollection& Collection, int32 GeometryIdx, bool bVisible)
+	{
+		int32 FaceEnd = Collection.FaceCount[GeometryIdx] + Collection.FaceStart[GeometryIdx];
+		for (int32 FaceIdx = Collection.FaceStart[GeometryIdx]; FaceIdx < FaceEnd; FaceIdx++)
+		{
+			Collection.Visible[FaceIdx] = bVisible;
+		}
+	}
+
+	static int32 AppendToCollection(const FTransform& ToCollection, FAugmentedDynamicMesh& AugMesh, int32 TransformParent, int32 SubPartIndex, FGeometryCollection& Output)
+	{
+		FDynamicMesh3& Mesh = AugMesh.Mesh;
+		if (Mesh.TriangleCount() == 0)
+		{
+			return -1;
+		}
+
+		if (!Mesh.IsCompact())
+		{
+			Mesh.CompactInPlace(nullptr);
+		}
+
+		int32 NewGeometryStartIdx = Output.FaceStart.Num();
+		int32 OriginalVertexNum = Output.Vertex.Num();
+		int32 OriginalFaceNum = Output.Indices.Num();
+
+		int32 GeometryIdx = Output.AddElements(1, FGeometryCollection::GeometryGroup);
+		int32 TransformIdx = Output.AddElements(1, FGeometryCollection::TransformGroup);
+
+		int32 NumTriangles = Mesh.TriangleCount();
+		int32 NumVertices = Mesh.VertexCount();
+		check(NumTriangles > 0);
+		check(Mesh.IsCompact());
+		Output.FaceCount[GeometryIdx] = NumTriangles;
+		Output.FaceStart[GeometryIdx] = OriginalFaceNum;
+		Output.VertexCount[GeometryIdx] = NumVertices;
+		Output.VertexStart[GeometryIdx] = OriginalVertexNum;
+		Output.TransformIndex[GeometryIdx] = TransformIdx;
+		Output.TransformToGeometryIndex[TransformIdx] = GeometryIdx;
+		if (TransformParent > -1)
+		{
+			// TODO: this is probably not the best way to build the bone name string?
+			Output.BoneName[TransformIdx] = Output.BoneName[TransformParent] + "_" + FString::FromInt(SubPartIndex);
+			Output.BoneColor[TransformIdx] = Output.BoneColor[TransformParent];
+			Output.Parent[TransformIdx] = TransformParent;
+			Output.Children[TransformParent].Add(TransformIdx);
+		}
+		Output.Transform[TransformIdx] = FTransform::Identity;
+
+		int32 FacesStart = Output.AddElements(NumTriangles, FGeometryCollection::FacesGroup);
+		int32 VerticesStart = Output.AddElements(NumVertices, FGeometryCollection::VerticesGroup);
+
+		for (int32 VID = 0; VID < Mesh.MaxVertexID(); VID++)
+		{
+			checkSlow(Mesh.IsVertex(VID)); // mesh is compact
+			int32 CopyToIdx = VerticesStart + VID;
+			Output.Vertex[CopyToIdx] = ToCollection.TransformPosition(FVector(Mesh.GetVertex(VID)));
+			Output.Normal[CopyToIdx] = ToCollection.TransformVectorNoScale(FVector(Mesh.GetVertexNormal(VID))); // TODO: use correct normal transform instead
+			Output.UV[CopyToIdx] = FVector2D(Mesh.GetVertexUV(VID));
+			FVector3f TangentU, TangentV;
+			AugMesh.GetTangent(VID, TangentU, TangentV);
+			Output.TangentU[CopyToIdx] = ToCollection.TransformVectorNoScale(FVector(TangentU));
+			Output.TangentV[CopyToIdx] = ToCollection.TransformVectorNoScale(FVector(TangentV));
+			Output.Color[CopyToIdx] = FVector(Mesh.GetVertexColor(VID));
+
+			// Bone map is set based on the transform of the new geometry
+			Output.BoneMap[CopyToIdx] = TransformIdx;
+		}
+
+		FIntVector VertexStartOffset(VerticesStart);
+		for (int32 TID = 0; TID < Mesh.MaxTriangleID(); TID++)
+		{
+			check(Mesh.IsTriangle(TID));
+			int32 CopyToIdx = FacesStart + TID;
+			Output.Visible[CopyToIdx] = AugMesh.GetVisibility(TID);
+			Output.MaterialID[CopyToIdx] = Mesh.Attributes()->GetMaterialID()->GetValue(TID);
+			Output.Indices[CopyToIdx] = FIntVector(Mesh.GetTriangle(TID)) + VertexStartOffset;
+		}
+
+		if (Output.BoundingBox.Num())
+		{
+			Output.BoundingBox[GeometryIdx].Init();
+			for (int32 Idx = OriginalVertexNum; Idx < Output.Vertex.Num(); ++Idx)
+			{
+				Output.BoundingBox[GeometryIdx] += Output.Vertex[Idx];
+			}
+		}
+
+		return GeometryIdx;
+	}
+};
+
+
 // logic from FMeshUtility::GenerateGeometryCollectionFromBlastChunk, sets material IDs based on construction pattern that external materials have even IDs and are matched to internal materials at InternalID = ExternalID+1
-int32 FInternalSurfaceMaterials::GetDefaultMaterialIDForGeometry(const FGeometryCollection& Collection, int32 GeometryIdx)
+int32 FInternalSurfaceMaterials::GetDefaultMaterialIDForGeometry(const FGeometryCollection& Collection, int32 GeometryIdx) const
 {
 	int32 FaceStart = 0;
 	int32 FaceEnd = Collection.Indices.Num();
@@ -149,6 +1328,10 @@
 	{
 		GlobalUVScale =  UVDistance / WorldDistance;
 	}
+	if (GlobalUVScale <= 0)
+	{
+		GlobalUVScale = 1;
+	}
 }
 
 inline double PlaneDotDouble(const FPlane& Plane, const FVector& V)
@@ -1891,14 +3074,7 @@
 			for (FIntVector& Face : PlaneTriangulation)
 			{
 				FVector3d TriCentroid = (Triangulation.LocalVertices[Face.X] + Triangulation.LocalVertices[Face.Y] + Triangulation.LocalVertices[Face.Z]) / 3.0;
-				//float WindingDebug = FastWindingTree.SlowWindingNumber(TriCentroid);
 				float WindingFast = FastWindingTree.FastWindingNumber(TriCentroid);
-				/*if (FMath::Abs(WindingDebug - WindingFast) > 1e-2)
-				{
-					std::cerr << WindingFast << " " << WindingDebug << std::endl;
-					float WindingFast = FastWindingTree.FastWindingNumber(TriCentroid);
-					std::cerr << "always";
-				}*/
 				if (WindingFast > .5)
 				{
 					Triangulation.LocalIndices.Add(Face);
@@ -2149,6 +3325,17 @@
 	return CutMultipleWithPlanarCells(Cells, Source, TransformIndices, TransformCells, bIncludeOutsideCellInOutput, CheckDistanceAcrossOutsideCellForProximity, bSetDefaultInternalMaterialsFromCollection, VertexInterpolate);
 }
 
+int32 CutMultipleWithPlanarCells_MeshBooleanPath(
+	FPlanarCells& Cells,
+	FGeometryCollection& Source,
+	const TArrayView<const int32>& TransformIndices,
+	const TOptional<FTransform>& TransformCells,
+	bool bIncludeOutsideCellInOutput,
+	float CheckDistanceAcrossOutsideCellForProximity,
+	bool bSetDefaultInternalMaterialsFromCollection,
+	TFunction<void(const FGeometryCollection&, int32, const FGeometryCollection&, int32, float, int32, FGeometryCollection&)> VertexInterpolate
+);
+
 // Cut multiple Geometry groups inside a GeometryCollection with PlanarCells, and add each cut cell back to the GeometryCollection as a new child of their source Geometry
 int32 CutMultipleWithPlanarCells(
 	FPlanarCells &Cells,
@@ -2161,6 +3348,18 @@
 	TFunction<void(const FGeometryCollection&, int32, const FGeometryCollection&, int32, float, int32, FGeometryCollection&)> VertexInterpolate
 )
 {
+	if (CVarFractureMethod.GetValueOnAnyThread() >= 1)
+	{
+		return CutMultipleWithPlanarCells_MeshBooleanPath(
+			Cells,
+			Source,
+			TransformIndices,
+			TransformCells,
+			bIncludeOutsideCellInOutput,
+			CheckDistanceAcrossOutsideCellForProximity,
+			bSetDefaultInternalMaterialsFromCollection,
+			VertexInterpolate);
+	}
 	float PlaneEps = 1e-4;
 
 	int32 NewGeomStartIdx = -1;
@@ -2259,17 +3458,19 @@
 }
 
 
-int32 CutMultipleWithMultiplePlanes(
+
+
+int32 CutMultipleWithMultiplePlanes_old(
 	const TArrayView<const FPlane>& Planes,
 	FInternalSurfaceMaterials& InternalSurfaceMaterials,
 	FGeometryCollection& Collection,
 	const TArrayView<const int32>& TransformIndices,
 	const TOptional<FTransform>& TransformCells,
-	bool bFlattenToSingleLayer,
 	bool bSetDefaultInternalMaterialsFromCollection,
 	TFunction<void(const FGeometryCollection&, int32, const FGeometryCollection&, int32, float, int32, FGeometryCollection&)> VertexInterpolate
 )
 {
+	bool bFlattenToSingleLayer = true;
 	float PlaneEps = 1e-4;
 
 	int32 OrigNumGeom = Collection.FaceCount.Num();
@@ -2550,3 +3751,93 @@
 
 	return OrigNumGeom == Collection.FaceCount.Num() ? -1 : OrigNumGeom;
 }
+
+
+int32 CutMultipleWithMultiplePlanes(
+	const TArrayView<const FPlane>& Planes,
+	FInternalSurfaceMaterials& InternalSurfaceMaterials,
+	FGeometryCollection& Collection,
+	const TArrayView<const int32>& TransformIndices,
+	const TOptional<FTransform>& TransformCells,
+	bool bSetDefaultInternalMaterialsFromCollection,
+	TFunction<void(const FGeometryCollection&, int32, const FGeometryCollection&, int32, float, int32, FGeometryCollection&)> VertexInterpolate
+)
+{
+	if (CVarFractureMethod.GetValueOnAnyThread() <= 0)
+	{
+		return CutMultipleWithMultiplePlanes_old(Planes, InternalSurfaceMaterials, Collection, TransformIndices, TransformCells, bSetDefaultInternalMaterialsFromCollection, VertexInterpolate);
+	}
+	int32 OrigNumGeom = Collection.FaceCount.Num();
+	int32 CurNumGeom = OrigNumGeom;
+
+	if (bSetDefaultInternalMaterialsFromCollection)
+	{
+		InternalSurfaceMaterials.SetUVScaleFromCollection(Collection);
+	}
+
+	if (!Collection.HasAttribute("Proximity", FGeometryCollection::GeometryGroup))
+	{
+		const FManagedArrayCollection::FConstructionParameters GeometryDependency(FGeometryCollection::GeometryGroup);
+		Collection.AddAttribute<TSet<int32>>("Proximity", FGeometryCollection::GeometryGroup, GeometryDependency);
+	}
+
+	FTransform CellsToWorld = TransformCells.Get(FTransform::Identity);
+	FTransform WorldToCells = CellsToWorld.Inverse();
+
+	FDynamicMeshCollection MeshCollection(&Collection, TransformIndices, CellsToWorld);
+	double OnePercentExtend = MeshCollection.Bounds.MaxDim() * .01;
+
+	int32 NewGeomStartIdx = -1;
+	NewGeomStartIdx = MeshCollection.CutWithMultiplePlanes(Planes, &Collection, InternalSurfaceMaterials, bSetDefaultInternalMaterialsFromCollection);
+
+	Collection.ReindexMaterials();
+	return NewGeomStartIdx;
+}
+
+
+// Cut multiple Geometry groups inside a GeometryCollection with PlanarCells, and add each cut cell back to the GeometryCollection as a new child of their source Geometry
+int32 CutMultipleWithPlanarCells_MeshBooleanPath(
+	FPlanarCells& Cells,
+	FGeometryCollection& Source,
+	const TArrayView<const int32>& TransformIndices,
+	const TOptional<FTransform>& TransformCells,
+	bool bIncludeOutsideCellInOutput,
+	float CheckDistanceAcrossOutsideCellForProximity,
+	bool bSetDefaultInternalMaterialsFromCollection,
+	TFunction<void(const FGeometryCollection&, int32, const FGeometryCollection&, int32, float, int32, FGeometryCollection&)> VertexInterpolate
+)
+{
+	if (!Source.HasAttribute("Proximity", FGeometryCollection::GeometryGroup))
+	{
+		const FManagedArrayCollection::FConstructionParameters GeometryDependency(FGeometryCollection::GeometryGroup);
+		Source.AddAttribute<TSet<int32>>("Proximity", FGeometryCollection::GeometryGroup, GeometryDependency);
+	}
+
+	if (bSetDefaultInternalMaterialsFromCollection)
+	{
+		Cells.InternalSurfaceMaterials.SetUVScaleFromCollection(Source);
+	}
+
+	FTransform CellsToWorld = TransformCells.Get(FTransform::Identity);
+
+	FDynamicMeshCollection MeshCollection(&Source, TransformIndices, CellsToWorld);
+	double OnePercentExtend = MeshCollection.Bounds.MaxDim() * .01;
+	FCellMeshes CellMeshes(Cells, MeshCollection.Bounds, OnePercentExtend, bIncludeOutsideCellInOutput);
+
+	int32 NewGeomStartIdx = -1;
+
+	NewGeomStartIdx = MeshCollection.CutWithCellMeshes(Cells, CellMeshes, &Source, bSetDefaultInternalMaterialsFromCollection);
+
+	Source.ReindexMaterials();
+	return NewGeomStartIdx;
+}
+
+
+
+
+
+
+
+
+
+
diff --git a/Engine/Plugins/Experimental/PlanarCutPlugin/Source/PlanarCut/Public/PlanarCut.h b/Engine/Plugins/Experimental/PlanarCutPlugin/Source/PlanarCut/Public/PlanarCut.h
--- a/Engine/Plugins/Experimental/PlanarCutPlugin/Source/PlanarCut/Public/PlanarCut.h
+++ b/Engine/Plugins/Experimental/PlanarCutPlugin/Source/PlanarCut/Public/PlanarCut.h
@@ -33,7 +33,7 @@ struct PLANARCUT_API FInternalSurfaceMaterials
 	void SetUVScaleFromCollection(const FGeometryCollection& Collection, int32 GeometryIdx = -1);
 
 	
-	int32 GetDefaultMaterialIDForGeometry(const FGeometryCollection& Collection, int32 GeometryIdx = -1);
+	int32 GetDefaultMaterialIDForGeometry(const FGeometryCollection& Collection, int32 GeometryIdx = -1) const;
 };
 
 // Stores planar facets that divide space into cells
@@ -60,6 +60,14 @@ struct PLANARCUT_API FPlanarCells
 	TFunction<int32(FVector)> CellFromPosition;
 
 	FInternalSurfaceMaterials InternalSurfaceMaterials;
+
+	/**
+	 * @return true if this is a single, unbounded cutting plane
+	 */
+	bool IsInfinitePlane() const
+	{
+		return NumCells == 2 && Planes.Num() == 1 && PlaneBoundaries[0].Num() == 0;
+	}
 	
 	/**
 	 * Debugging function to check that the plane boundary vertices are wound to match the orientation of the plane normal vectors
@@ -183,7 +191,7 @@ int32 PLANARCUT_API CutMultipleWithPlanarCells(
 	const TArrayView<const int32>& TransformIndices,
 	const TOptional<FTransform>& TransformCells = TOptional<FTransform>(),
 	bool bIncludeOutsideCellInOutput = true,
-	float CheckDistanceAcrossOutsideCellForProximity = 0,
+	float CheckDistanceAcrossOutsideCellForProximity = 0,  // TODO: < this param does nothing in the new mode; is only needed in special cases that aren't possible in the UI currently
 	bool bSetDefaultInternalMaterialsFromCollection = true,
 	TFunction<void(const FGeometryCollection&, int32, const FGeometryCollection&, int32, float, int32, FGeometryCollection&)> VertexInterpolate = DefaultVertexInterpolation
 );
@@ -196,7 +204,6 @@ int32 PLANARCUT_API CutMultipleWithPlanarCells(
  * @param Collection			The collection to be cut
  * @param TransformIndices	Which transform groups inside the collection to cut
  * @param TransformCells	Optional transform of the planar cut; if unset, defaults to Identity
- * @param bFlattenToSingleLayer If true, only add one layer overall to the hierarchy even if some geometry is cut multiple times.
  * @param VertexInterpolate	Function that interpolates vertex properties (UVs, normals, etc); a default that handles all the normal vertex properties is provided, should only need to replace this if you have custom attributes
  * @return	index of first new geometry in the Output GeometryCollection, or -1 if no geometry was added
  */
@@ -206,7 +213,6 @@ int32 PLANARCUT_API CutMultipleWithMultiplePlanes(
 	FGeometryCollection& Collection,
 	const TArrayView<const int32>& TransformIndices,
 	const TOptional<FTransform>& TransformCells = TOptional<FTransform>(),
-	bool bFlattenToSingleLayer = true,
 	bool bSetDefaultInternalMaterialsFromCollection = true,
 	TFunction<void(const FGeometryCollection&, int32, const FGeometryCollection&, int32, float, int32, FGeometryCollection&)> VertexInterpolate = DefaultVertexInterpolation
-);
\ No newline at end of file
+);
