Avoid creating too many new vertices when adding surface resolution for noise on cuts.  Also set an absolute min on the UI and fix the surface resolution type (should be float)

#rnx

diff --git a/Engine/Plugins/Experimental/ChaosEditor/Source/FractureEditor/Public/FractureTool.h b/Engine/Plugins/Experimental/ChaosEditor/Source/FractureEditor/Public/FractureTool.h
--- a/Engine/Plugins/Experimental/ChaosEditor/Source/FractureEditor/Public/FractureTool.h
+++ b/Engine/Plugins/Experimental/ChaosEditor/Source/FractureEditor/Public/FractureTool.h
@@ -76,8 +76,8 @@ public:
 	int32 OctaveNumber;
 
 	/** Spacing between vertices on cut surfaces, where noise is added.  Larger spacing between vertices will create more efficient meshes with fewer triangles, but less resolution to see the shape of the added noise  */
-	UPROPERTY(EditAnywhere, Category = Noise, meta = (UIMin = "1"))
-	int32 SurfaceResolution;
+	UPROPERTY(EditAnywhere, Category = Noise, meta = (UIMin = "1", ClampMin = "0.1"))
+	float SurfaceResolution;
 
 	UPROPERTY()
 	UFractureTool *OwnerTool;
diff --git a/Engine/Plugins/Experimental/PlanarCutPlugin/Source/PlanarCut/Private/PlanarCut.cpp b/Engine/Plugins/Experimental/PlanarCutPlugin/Source/PlanarCut/Private/PlanarCut.cpp
--- a/Engine/Plugins/Experimental/PlanarCutPlugin/Source/PlanarCut/Private/PlanarCut.cpp
+++ b/Engine/Plugins/Experimental/PlanarCutPlugin/Source/PlanarCut/Private/PlanarCut.cpp
@@ -499,6 +499,22 @@ private:
 			}
 		}
 	}
+
+	// Approximately calculate a "safe" spacing that would not require the remesher to create more than a million new vertices
+	double GetSafeNoiseSpacing(float SurfaceArea, float TargetSpacing)
+	{
+		double MaxVerts = 1000000;
+		double MinEdgeLen = FMathd::Sqrt((double)SurfaceArea / MaxVerts);
+		double Spacing = FMath::Max3(.001, MinEdgeLen, (double)TargetSpacing);
+		if (Spacing > TargetSpacing)
+		{
+			UE_LOG(LogPlanarCut, Warning,
+				TEXT("Requested spacing of noise points (surface resolution) of %f would require too many added vertices; Using %f instead."),
+				TargetSpacing, Spacing);
+		}
+		return Spacing;
+	}
+
 	void CreateMeshesForBoundedPlanesWithNoise(int NumCells, const FPlanarCells& Cells, const FAxisAlignedBox3d& DomainBounds, bool bNoise, double GlobalUVScale)
 	{
 		TArray<FDynamicMesh3> PlaneMeshes;
@@ -527,7 +543,7 @@ private:
 		}
 
 		// heuristic to protect against creating too many vertices on remeshing
-		float MaxArea = 0;
+		float TotalArea = 0;
 		for (int32 PlaneIdx = 0; PlaneIdx < Cells.Planes.Num(); PlaneIdx++)
 		{
 			const TArray<int>& PlaneBoundary = Cells.PlaneBoundaries[PlaneIdx];
@@ -539,11 +555,9 @@ private:
 				const FVector& V2 = Cells.PlaneBoundaryVertices[PlaneBoundary[V2Idx]];
 				AreaVec += (V1 - V0) ^ (V2 - V1);
 			}
-			MaxArea = AreaVec.Size();
+			TotalArea += AreaVec.Size();
 		}
-		double MaxVerts = 10000;
-		double MinEdgeLen = FMathd::Sqrt((double)MaxArea / MaxVerts);
-		double Spacing = FMath::Max3(.001, MinEdgeLen, (double)Cells.InternalSurfaceMaterials.NoiseSettings->PointSpacing);
+		double Spacing = GetSafeNoiseSpacing(TotalArea, Cells.InternalSurfaceMaterials.NoiseSettings->PointSpacing);
 
 		ParallelFor(Cells.Planes.Num(), [this, &PlaneMeshes, &Cells, GlobalUVScale, Spacing](int32 PlaneIdx)
 			{
@@ -643,7 +657,8 @@ private:
 
 		if (bNoise)
 		{
-			RemeshForNoise(PlaneMesh, EEdgeRefineFlags::SplitsOnly, Cells.InternalSurfaceMaterials.NoiseSettings->PointSpacing);
+			double Spacing = GetSafeNoiseSpacing(XYRange.Area(), Cells.InternalSurfaceMaterials.NoiseSettings->PointSpacing);
+			RemeshForNoise(PlaneMesh, EEdgeRefineFlags::SplitsOnly, Spacing);
 			ApplyNoise(PlaneMesh, PlaneFrame.GetAxis(2), Cells.InternalSurfaceMaterials.NoiseSettings.GetValue(), true);
 			FMeshNormals::QuickComputeVertexNormals(PlaneMesh);
 		}
