Fix vertex UVs and other attributes getting lost on fracture

diff --git a/Engine/Plugins/Experimental/GeometryProcessing/Source/DynamicMesh/Private/DynamicMesh3_Edits.cpp b/Engine/Plugins/Experimental/GeometryProcessing/Source/DynamicMesh/Private/DynamicMesh3_Edits.cpp
--- a/Engine/Plugins/Experimental/GeometryProcessing/Source/DynamicMesh/Private/DynamicMesh3_Edits.cpp
+++ b/Engine/Plugins/Experimental/GeometryProcessing/Source/DynamicMesh/Private/DynamicMesh3_Edits.cpp
@@ -1395,6 +1395,21 @@ EMeshResult FDynamicMesh3::CollapseEdge(int vKeep, int vRemove, double collapse_
 	// save vertex positions before we delete removed (can defer kept?)
 	FVector3d KeptPos = GetVertex(vKeep);
 	FVector3d RemovedPos = GetVertex(vRemove);
+	FVector2f RemovedUV;
+	if (HasVertexUVs())
+	{
+		RemovedUV = GetVertexUV(vRemove);
+	}
+	FVector3f RemovedNormal;
+	if (HasVertexNormals())
+	{
+		RemovedNormal = GetVertexNormal(vRemove);
+	}
+	FVector3f RemovedColor;
+	if (HasVertexColors())
+	{
+		RemovedColor = GetVertexColor(vRemove);
+	}
 
 	// 1) remove edge ab from vtx b
 	// 2) find edges ad and ac, and tris tad, tac across those edges  (will use later)
@@ -1567,6 +1582,18 @@ EMeshResult FDynamicMesh3::CollapseEdge(int vKeep, int vRemove, double collapse_
 
 	// set kept vertex to interpolated collapse position
 	SetVertex(vKeep, FVector3d::Lerp(KeptPos, RemovedPos, collapse_t));
+	if (HasVertexUVs())
+	{
+		SetVertexUV(vKeep, FVector2f::Lerp(GetVertexUV(vKeep), RemovedUV, collapse_t));
+	}
+	if (HasVertexNormals())
+	{
+		SetVertexNormal(vKeep, FVector3f::Lerp(GetVertexNormal(vKeep), RemovedNormal, collapse_t).Normalized());
+	}
+	if (HasVertexColors())
+	{
+		SetVertexColor(vKeep, FVector3f::Lerp(GetVertexColor(vKeep), RemovedColor, collapse_t));
+	}
 
 	CollapseInfo.KeptVertex = vKeep;
 	CollapseInfo.RemovedVertex = vRemove;
diff --git a/Engine/Plugins/Experimental/GeometryProcessing/Source/DynamicMesh/Private/DynamicMeshEditor.cpp b/Engine/Plugins/Experimental/GeometryProcessing/Source/DynamicMesh/Private/DynamicMeshEditor.cpp
--- a/Engine/Plugins/Experimental/GeometryProcessing/Source/DynamicMesh/Private/DynamicMeshEditor.cpp
+++ b/Engine/Plugins/Experimental/GeometryProcessing/Source/DynamicMesh/Private/DynamicMeshEditor.cpp
@@ -1415,6 +1415,12 @@ void FDynamicMeshEditor::AppendMesh(const FDynamicMesh3* AppendMesh,
 			Mesh->SetVertexNormal(NewVertID, Normal);
 		}
 
+		if (AppendMesh->HasVertexUVs() && Mesh->HasVertexUVs())
+		{
+			FVector2f UV = AppendMesh->GetVertexUV(VertID);
+			Mesh->SetVertexUV(NewVertID, UV);
+		}
+
 		if (AppendMesh->HasVertexColors() && Mesh->HasVertexColors())
 		{
 			FVector3f Color = AppendMesh->GetVertexColor(VertID);
diff --git a/Engine/Plugins/Experimental/PlanarCutPlugin/Source/PlanarCut/Private/PlanarCut.cpp b/Engine/Plugins/Experimental/PlanarCutPlugin/Source/PlanarCut/Private/PlanarCut.cpp
--- a/Engine/Plugins/Experimental/PlanarCutPlugin/Source/PlanarCut/Private/PlanarCut.cpp
+++ b/Engine/Plugins/Experimental/PlanarCutPlugin/Source/PlanarCut/Private/PlanarCut.cpp
@@ -111,7 +111,7 @@ struct FAugmentedDynamicMesh
 	const FName TangentVAttribName = "TangentVAttrib";
 	const FName VisibleAttribName = "VisibleAttrib";
 
-	FAugmentedDynamicMesh()
+	FAugmentedDynamicMesh() : Mesh(true, true, true, true)
 	{
 		Mesh.EnableAttributes();
 		Mesh.Attributes()->EnableMaterialID();
@@ -125,6 +125,10 @@ struct FAugmentedDynamicMesh
 	// re-setup attributes to augment the mesh; only call this if we've lost the attached attributes
 	void Augment()
 	{
+		Mesh.EnableVertexColors(FVector3f(1, 1, 1));
+		Mesh.EnableVertexNormals(FVector3f::UnitZ());
+		Mesh.EnableVertexUVs(FVector2f(0, 0));
+		Mesh.EnableTriangleGroups();
 		Mesh.EnableAttributes();
 		Mesh.Attributes()->EnableMaterialID();
 		ensure(Mesh.Attributes()->NumAttachedAttributes() == 0);
@@ -519,6 +523,13 @@ private:
 	{
 		TArray<FDynamicMesh3> PlaneMeshes;
 		PlaneMeshes.SetNum(Cells.Planes.Num());
+		for (FDynamicMesh3& PlaneMesh : PlaneMeshes)
+		{
+			PlaneMesh.EnableTriangleGroups();
+			PlaneMesh.EnableVertexUVs(FVector2f(0, 0));
+			PlaneMesh.EnableVertexNormals(FVector3f::UnitZ());
+			PlaneMesh.EnableVertexColors(FVector3f(1, 1, 1));
+		}
 
 		struct FPlaneIdxAndFlip
 		{
@@ -638,7 +649,7 @@ private:
 		//	// TODO: early out for plane that doesn't even intersect the domain bounding box?
 		//}
 
-		FDynamicMesh3 PlaneMesh;
+		FDynamicMesh3 PlaneMesh(true, true, true, true);
 		FVertexInfo PlaneVertInfo;
 		PlaneVertInfo.bHaveC = true;
 		PlaneVertInfo.bHaveUV = true;
