split fracture results into disconnected pieces as needed
+ fix proximity-finding method to not use triangle group IDs (since those change on mesh append, which happens during mesh booleans)

diff --git a/Engine/Plugins/Experimental/GeometryProcessing/Source/DynamicMesh/Private/DynamicMesh3.cpp b/Engine/Plugins/Experimental/GeometryProcessing/Source/DynamicMesh/Private/DynamicMesh3.cpp
--- a/Engine/Plugins/Experimental/GeometryProcessing/Source/DynamicMesh/Private/DynamicMesh3.cpp
+++ b/Engine/Plugins/Experimental/GeometryProcessing/Source/DynamicMesh/Private/DynamicMesh3.cpp
@@ -366,6 +366,42 @@ int FDynamicMesh3::GetComponentsFlags() const
 	return c;
 }
 
+void FDynamicMesh3::EnableMeshComponents(int MeshComponentsFlags)
+{
+	if (int(EMeshComponents::FaceGroups) & MeshComponentsFlags)
+	{
+		EnableTriangleGroups(0);
+	}
+	else
+	{
+		DiscardTriangleGroups();
+	}
+	if (int(EMeshComponents::VertexColors) & MeshComponentsFlags)
+	{
+		EnableVertexColors(FVector3f(1, 1, 1));
+	}
+	else
+	{
+		DiscardVertexColors();
+	}
+	if (int(EMeshComponents::VertexNormals) & MeshComponentsFlags)
+	{
+		EnableVertexNormals(FVector3f::UnitY());
+	}
+	else
+	{
+		DiscardVertexNormals();
+	}
+	if (int(EMeshComponents::VertexUVs) & MeshComponentsFlags)
+	{
+		EnableVertexUVs(FVector2f(0, 0));
+	}
+	else
+	{
+		DiscardVertexUVs();
+	}
+}
+
 void FDynamicMesh3::EnableVertexNormals(const FVector3f& InitialNormal)
 {
 	if (HasVertexNormals())
diff --git a/Engine/Plugins/Experimental/GeometryProcessing/Source/DynamicMesh/Private/DynamicMeshEditor.cpp b/Engine/Plugins/Experimental/GeometryProcessing/Source/DynamicMesh/Private/DynamicMeshEditor.cpp
--- a/Engine/Plugins/Experimental/GeometryProcessing/Source/DynamicMesh/Private/DynamicMeshEditor.cpp
+++ b/Engine/Plugins/Experimental/GeometryProcessing/Source/DynamicMesh/Private/DynamicMeshEditor.cpp
@@ -1783,6 +1783,7 @@ bool FDynamicMeshEditor::SplitMesh(const FDynamicMesh3* SourceMesh, TArray<FDyna
 	// enable matching attributes
 	for (FDynamicMesh3& M : SplitMeshes)
 	{
+		M.EnableMeshComponents(SourceMesh->GetComponentsFlags());
 		if (SourceMesh->HasAttributes())
 		{
 			M.EnableAttributes();
diff --git a/Engine/Plugins/Experimental/GeometryProcessing/Source/DynamicMesh/Public/DynamicMesh3.h b/Engine/Plugins/Experimental/GeometryProcessing/Source/DynamicMesh/Public/DynamicMesh3.h
--- a/Engine/Plugins/Experimental/GeometryProcessing/Source/DynamicMesh/Public/DynamicMesh3.h
+++ b/Engine/Plugins/Experimental/GeometryProcessing/Source/DynamicMesh/Public/DynamicMesh3.h
@@ -608,6 +608,13 @@ public:
 	// Vertex and Triangle attribute arrays
 	//
 public:
+	/**
+	 * Enable requested set of mesh components (triangle groups and vertex normals/colors/UVs)
+	 * and discard any that are not requested
+	 * @param MeshComponentsFlags A 'bitwise or' of requested EMeshComponents flags
+	 */
+	void EnableMeshComponents(int MeshComponentsFlags);
+
 	void EnableVertexNormals(const FVector3f& InitialNormal);
 	void DiscardVertexNormals();
 
diff --git a/Engine/Plugins/Experimental/GeometryProcessing/Source/GeometricObjects/Public/Spatial/PointHashGrid3.h b/Engine/Plugins/Experimental/GeometryProcessing/Source/GeometricObjects/Public/Spatial/PointHashGrid3.h
--- a/Engine/Plugins/Experimental/GeometryProcessing/Source/GeometricObjects/Public/Spatial/PointHashGrid3.h
+++ b/Engine/Plugins/Experimental/GeometryProcessing/Source/GeometricObjects/Public/Spatial/PointHashGrid3.h
@@ -203,6 +203,63 @@ public:
 
 		return TPair<PointDataType, RealType>(nearest, min_distsq);
 	}
+
+
+	/**
+	 * Find all points in grid within a given sphere, without locking / thread-safety.
+	 * @param QueryPoint the center of the query sphere
+	 * @param Radius the radius of the query sphere
+	 * @param DistanceSqFunc Function you provide which measures the squared distance between QueryPoint and a Value
+	 * @param ResultsOut Array of output points in sphere
+	 * @param IgnoreFunc optional Function you may provide which will result in a Value being ignored if IgnoreFunc(Value) returns true
+	 * @return the number of found points
+	 */
+	int FindPointsInBall(
+		const FVector3<RealType>& QueryPoint, RealType Radius,
+		TFunctionRef<RealType(const PointDataType&)> DistanceSqFunc,
+		TArray<PointDataType>& ResultsOut,
+		TFunctionRef<bool(const PointDataType&)> IgnoreFunc = [](const PointDataType& data) { return false; }) const
+	{
+		if (!Hash.Num())
+		{
+			return 0;
+		}
+
+		FVector3i min_idx = Indexer.ToGrid(QueryPoint - Radius * FVector3<RealType>::One());
+		FVector3i max_idx = Indexer.ToGrid(QueryPoint + Radius * FVector3<RealType>::One());
+
+		RealType RadiusSquared = Radius * Radius;
+
+		TArray<PointDataType> Values;
+		for (int zi = min_idx.Z; zi <= max_idx.Z; zi++)
+		{
+			for (int yi = min_idx.Y; yi <= max_idx.Y; yi++)
+			{
+				for (int xi = min_idx.X; xi <= max_idx.X; xi++)
+				{
+					FVector3i idx(xi, yi, zi);
+					Values.Reset();
+					Hash.MultiFind(idx, Values);
+					for (PointDataType Value : Values)
+					{
+						if (IgnoreFunc(Value))
+						{
+							continue;
+						}
+						RealType distsq = DistanceSqFunc(Value);
+						if (distsq < RadiusSquared)
+						{
+							ResultsOut.Add(Value);
+						}
+					}
+				}
+			}
+		}
+
+		return ResultsOut.Num();
+	}
+
+
 };
 
 template <typename PointDataType> using TPointHashGrid3d = TPointHashGrid3<PointDataType, double>;
diff --git a/Engine/Plugins/Experimental/PlanarCutPlugin/Source/PlanarCut/PlanarCut.Build.cs b/Engine/Plugins/Experimental/PlanarCutPlugin/Source/PlanarCut/PlanarCut.Build.cs
--- a/Engine/Plugins/Experimental/PlanarCutPlugin/Source/PlanarCut/PlanarCut.Build.cs
+++ b/Engine/Plugins/Experimental/PlanarCutPlugin/Source/PlanarCut/PlanarCut.Build.cs
@@ -8,7 +8,7 @@ namespace UnrealBuildTool.Rules
 		{
 			PrivateIncludePaths.Add("PlanarCut/Private");
             PublicIncludePaths.Add(ModuleDirectory + "/Public");
-
+            PrivateDependencyModuleNames.Add("MeshUtilitiesCommon");
 			PublicDependencyModuleNames.AddRange(
 				new string[]
 				{
diff --git a/Engine/Plugins/Experimental/PlanarCutPlugin/Source/PlanarCut/Private/PlanarCut.cpp b/Engine/Plugins/Experimental/PlanarCutPlugin/Source/PlanarCut/Private/PlanarCut.cpp
--- a/Engine/Plugins/Experimental/PlanarCutPlugin/Source/PlanarCut/Private/PlanarCut.cpp
+++ b/Engine/Plugins/Experimental/PlanarCutPlugin/Source/PlanarCut/Private/PlanarCut.cpp
@@ -13,6 +13,7 @@
 #include "GeometryCollection/GeometryCollectionAlgo.h"
 #include "GeometryCollection/GeometryCollectionClusteringUtility.h"
 
+#include "DisjointSet.h"
 
 #include "DynamicMesh3.h"
 #include "DynamicMeshEditor.h"
@@ -103,139 +104,151 @@ struct FGeometryCollectionMeshAdapter
 	}
 };
 
-// FDynamicMesh3 + additional attributes as needed to carry standard geometry collection attributes
-struct FAugmentedDynamicMesh
+namespace UE
 {
-	FDynamicMesh3 Mesh;
-	const FName TangentUAttribName = "TangentUAttrib";
-	const FName TangentVAttribName = "TangentVAttrib";
-	const FName VisibleAttribName = "VisibleAttrib";
-
-	FAugmentedDynamicMesh() : Mesh(true, true, true, true)
+	namespace PlanarCutInternals
 	{
-		Mesh.EnableAttributes();
-		Mesh.Attributes()->EnableMaterialID();
-		Mesh.Attributes()->AttachAttribute(TangentUAttribName, new TDynamicMeshVertexAttribute<float, 3>(&Mesh));
-		Mesh.Attributes()->AttachAttribute(TangentVAttribName, new TDynamicMeshVertexAttribute<float, 1>(&Mesh)); // just a sign
-		TDynamicMeshScalarTriangleAttribute<bool>* VisAttrib = new TDynamicMeshScalarTriangleAttribute<bool>(&Mesh);
-		VisAttrib->Initialize(true);
-		Mesh.Attributes()->AttachAttribute(VisibleAttribName, VisAttrib);
-	}
+		// functions to setup geometry collection attributes on dynamic meshes
+		namespace AugmentDynamicMesh
+		{
+			FName TangentUAttribName = "TangentUAttrib";
+			FName TangentVAttribName = "TangentVAttrib";
+			FName VisibleAttribName = "VisibleAttrib";
 
-	// re-setup attributes to augment the mesh; only call this if we've lost the attached attributes
-	void Augment()
-	{
-		Mesh.EnableVertexColors(FVector3f(1, 1, 1));
-		Mesh.EnableVertexNormals(FVector3f::UnitZ());
-		Mesh.EnableVertexUVs(FVector2f(0, 0));
-		Mesh.EnableTriangleGroups();
-		Mesh.EnableAttributes();
-		Mesh.Attributes()->EnableMaterialID();
-		ensure(Mesh.Attributes()->NumAttachedAttributes() == 0);
-		Mesh.Attributes()->AttachAttribute(TangentUAttribName, new TDynamicMeshVertexAttribute<float, 3>(&Mesh));
-		Mesh.Attributes()->AttachAttribute(TangentVAttribName, new TDynamicMeshVertexAttribute<float, 1>(&Mesh)); // just a sign
-		TDynamicMeshScalarTriangleAttribute<bool>* VisAttrib = new TDynamicMeshScalarTriangleAttribute<bool>(&Mesh);
-		VisAttrib->Initialize(true);
-		Mesh.Attributes()->AttachAttribute(VisibleAttribName, VisAttrib);
-	}
+			void Augment(FDynamicMesh3& Mesh)
+			{
+				Mesh.EnableVertexColors(FVector3f(1, 1, 1));
+				Mesh.EnableVertexNormals(FVector3f::UnitZ());
+				Mesh.EnableVertexUVs(FVector2f(0, 0));
+				Mesh.EnableAttributes();
+				Mesh.Attributes()->EnableMaterialID();
+				ensure(Mesh.Attributes()->NumAttachedAttributes() == 0);
+				Mesh.Attributes()->AttachAttribute(TangentUAttribName, new TDynamicMeshVertexAttribute<float, 3>(&Mesh));
+				Mesh.Attributes()->AttachAttribute(TangentVAttribName, new TDynamicMeshVertexAttribute<float, 1>(&Mesh)); // just a sign
+				TDynamicMeshScalarTriangleAttribute<bool>* VisAttrib = new TDynamicMeshScalarTriangleAttribute<bool>(&Mesh);
+				VisAttrib->Initialize(true);
+				Mesh.Attributes()->AttachAttribute(VisibleAttribName, VisAttrib);
+			}
 
-	void SetDefaultMaterials(int32 MaterialID)
-	{
-		for (int TID : Mesh.TriangleIndicesItr())
-		{
-			Mesh.Attributes()->GetMaterialID()->SetNewValue(TID, MaterialID);
-		}
-	}
+			bool IsAugmented(const FDynamicMesh3& Mesh)
+			{
+				return Mesh.HasAttributes()
+					&& Mesh.Attributes()->NumAttachedAttributes() == 3
+					&& Mesh.Attributes()->HasAttachedAttribute(TangentUAttribName)
+					&& Mesh.Attributes()->HasAttachedAttribute(TangentVAttribName)
+					&& Mesh.Attributes()->HasAttachedAttribute(VisibleAttribName)
+					&& Mesh.Attributes()->HasMaterialID()
+					&& Mesh.HasVertexColors()
+					&& Mesh.HasVertexNormals()
+					&& Mesh.HasVertexUVs();
+			}
 
-	void SetDefaultAttributes(const FInternalSurfaceMaterials& Materials)
-	{
-		TDynamicMeshVertexAttribute<float, 3>* Us =
-			static_cast<TDynamicMeshVertexAttribute<float, 3>*>(Mesh.Attributes()->GetAttachedAttribute(TangentUAttribName));
-		TDynamicMeshVertexAttribute<float, 1>* VSigns =
-			static_cast<TDynamicMeshVertexAttribute<float, 1>*>(Mesh.Attributes()->GetAttachedAttribute(TangentVAttribName));
 
-		for (int VID : Mesh.VertexIndicesItr())
-		{
-			FVector3f N = Mesh.GetVertexNormal(VID);
-			FVector3f U;
-			VectorUtil::MakePerpVector(N, U);
-			Us->SetValue(VID, U);
-			float VSign = 1;
-			VSigns->SetValue(VID, &VSign);
-		}
+			void SetDefaultAttributes(FDynamicMesh3& Mesh, const FInternalSurfaceMaterials& Materials)
+			{
+				checkSlow(IsAugmented(Mesh));
+				TDynamicMeshVertexAttribute<float, 3>* Us =
+					static_cast<TDynamicMeshVertexAttribute<float, 3>*>(Mesh.Attributes()->GetAttachedAttribute(TangentUAttribName));
+				TDynamicMeshVertexAttribute<float, 1>* VSigns =
+					static_cast<TDynamicMeshVertexAttribute<float, 1>*>(Mesh.Attributes()->GetAttachedAttribute(TangentVAttribName));
 
-		TDynamicMeshScalarTriangleAttribute<bool>* Visible =
-			static_cast<TDynamicMeshScalarTriangleAttribute<bool>*>(Mesh.Attributes()->GetAttachedAttribute(VisibleAttribName));
-		for (int TID : Mesh.TriangleIndicesItr())
-		{
-			Visible->SetNewValue(TID, Materials.bGlobalVisibility);
-		}
-	}
+				for (int VID : Mesh.VertexIndicesItr())
+				{
+					FVector3f N = Mesh.GetVertexNormal(VID);
+					FVector3f U;
+					VectorUtil::MakePerpVector(N, U);
+					Us->SetValue(VID, U);
+					float VSign = 1;
+					VSigns->SetValue(VID, &VSign);
+				}
 
-	void SetVisibility(int TID, bool bIsVisible)
-	{
-		TDynamicMeshScalarTriangleAttribute<bool>* Visible =
-			static_cast<TDynamicMeshScalarTriangleAttribute<bool>*>(Mesh.Attributes()->GetAttachedAttribute(VisibleAttribName));
-		Visible->SetValue(TID, bIsVisible);
-	}
+				TDynamicMeshScalarTriangleAttribute<bool>* Visible =
+					static_cast<TDynamicMeshScalarTriangleAttribute<bool>*>(Mesh.Attributes()->GetAttachedAttribute(VisibleAttribName));
+				for (int TID : Mesh.TriangleIndicesItr())
+				{
+					Visible->SetNewValue(TID, Materials.bGlobalVisibility);
+				}
+			}
 
-	bool GetVisibility(int TID)
-	{
-		TDynamicMeshScalarTriangleAttribute<bool>* Visible =
-			static_cast<TDynamicMeshScalarTriangleAttribute<bool>*>(Mesh.Attributes()->GetAttachedAttribute(VisibleAttribName));
-		return Visible->GetValue(TID);
-	}
+			void SetVisibility(FDynamicMesh3& Mesh, int TID, bool bIsVisible)
+			{
+				checkSlow(IsAugmented(Mesh));
+				TDynamicMeshScalarTriangleAttribute<bool>* Visible =
+					static_cast<TDynamicMeshScalarTriangleAttribute<bool>*>(Mesh.Attributes()->GetAttachedAttribute(VisibleAttribName));
+				Visible->SetValue(TID, bIsVisible);
+			}
 
-	void SetTangent(int VID, FVector3f Normal, FVector3f TangentU, FVector3f TangentV)
-	{
-		TDynamicMeshVertexAttribute<float, 3>* Us =
-			static_cast<TDynamicMeshVertexAttribute<float, 3>*>(Mesh.Attributes()->GetAttachedAttribute(TangentUAttribName));
-		TDynamicMeshVertexAttribute<float, 1>* VSigns =
-			static_cast<TDynamicMeshVertexAttribute<float, 1>*>(Mesh.Attributes()->GetAttachedAttribute(TangentVAttribName));
-		float VSign = FMathf::SignNonZero(Normal.Cross(TangentU).Dot(TangentV));
-		Us->SetValue(VID, TangentU);
-		VSigns->SetValue(VID, &VSign);
-	}
-	void FixTangents()
-	{
-		TDynamicMeshVertexAttribute<float, 3>* Us =
-			static_cast<TDynamicMeshVertexAttribute<float, 3>*>(Mesh.Attributes()->GetAttachedAttribute(TangentUAttribName));
-		TDynamicMeshVertexAttribute<float, 1>* VSigns =
-			static_cast<TDynamicMeshVertexAttribute<float, 1>*>(Mesh.Attributes()->GetAttachedAttribute(TangentVAttribName));
+			bool GetVisibility(const FDynamicMesh3& Mesh, int TID)
+			{
+				checkSlow(IsAugmented(Mesh));
+				TDynamicMeshScalarTriangleAttribute<bool>* Visible =
+					static_cast<TDynamicMeshScalarTriangleAttribute<bool>*>(Mesh.Attributes()->GetAttachedAttribute(VisibleAttribName));
+				return Visible->GetValue(TID);
+			}
 
-		for (int VID : Mesh.VertexIndicesItr())
-		{
-			FVector3f N = Mesh.GetVertexNormal(VID);
-			FVector3f U;
-			Us->GetValue(VID, U);
-			U = N.Cross(U).Cross(N);
-			U.Normalize();
-			Us->SetValue(VID, U);
-			float VSign;
-			VSigns->GetValue(VID, &VSign);
-			VSign = FMathf::SignNonZero(VSign);
-			VSigns->SetValue(VID, &VSign);
-		}
-	}
-	void GetTangent(int VID, FVector3f& U, FVector3f& V)
-	{
-		TDynamicMeshVertexAttribute<float, 3>* Us =
-			static_cast<TDynamicMeshVertexAttribute<float, 3>*>(Mesh.Attributes()->GetAttachedAttribute(TangentUAttribName));
-		TDynamicMeshVertexAttribute<float, 1>* VSigns =
-			static_cast<TDynamicMeshVertexAttribute<float, 1>*>(Mesh.Attributes()->GetAttachedAttribute(TangentVAttribName));
-		FVector3f Normal = Mesh.GetVertexNormal(VID);
-		Us->GetValue(VID, U);
-		float VSign;
-		VSigns->GetValue(VID, &VSign);
-		V = Normal.Cross(U) * FMathf::SignNonZero(VSign);
+			void SetTangent(FDynamicMesh3& Mesh, int VID, FVector3f Normal, FVector3f TangentU, FVector3f TangentV)
+			{
+				checkSlow(IsAugmented(Mesh));
+				TDynamicMeshVertexAttribute<float, 3>* Us =
+					static_cast<TDynamicMeshVertexAttribute<float, 3>*>(Mesh.Attributes()->GetAttachedAttribute(TangentUAttribName));
+				TDynamicMeshVertexAttribute<float, 1>* VSigns =
+					static_cast<TDynamicMeshVertexAttribute<float, 1>*>(Mesh.Attributes()->GetAttachedAttribute(TangentVAttribName));
+				float VSign = FMathf::SignNonZero(Normal.Cross(TangentU).Dot(TangentV));
+				Us->SetValue(VID, TangentU);
+				VSigns->SetValue(VID, &VSign);
+			}
+
+			void FixTangents(FDynamicMesh3& Mesh)
+			{
+				checkSlow(IsAugmented(Mesh));
+				TDynamicMeshVertexAttribute<float, 3>* Us =
+					static_cast<TDynamicMeshVertexAttribute<float, 3>*>(Mesh.Attributes()->GetAttachedAttribute(TangentUAttribName));
+				TDynamicMeshVertexAttribute<float, 1>* VSigns =
+					static_cast<TDynamicMeshVertexAttribute<float, 1>*>(Mesh.Attributes()->GetAttachedAttribute(TangentVAttribName));
+
+				for (int VID : Mesh.VertexIndicesItr())
+				{
+					FVector3f N = Mesh.GetVertexNormal(VID);
+					FVector3f U;
+					Us->GetValue(VID, U);
+					U = N.Cross(U).Cross(N);
+					U.Normalize();
+					Us->SetValue(VID, U);
+					float VSign;
+					VSigns->GetValue(VID, &VSign);
+					VSign = FMathf::SignNonZero(VSign);
+					VSigns->SetValue(VID, &VSign);
+				}
+			}
+
+			void GetTangent(const FDynamicMesh3& Mesh, int VID, FVector3f& U, FVector3f& V)
+			{
+				checkSlow(IsAugmented(Mesh));
+				TDynamicMeshVertexAttribute<float, 3>* Us =
+					static_cast<TDynamicMeshVertexAttribute<float, 3>*>(Mesh.Attributes()->GetAttachedAttribute(TangentUAttribName));
+				TDynamicMeshVertexAttribute<float, 1>* VSigns =
+					static_cast<TDynamicMeshVertexAttribute<float, 1>*>(Mesh.Attributes()->GetAttachedAttribute(TangentVAttribName));
+				FVector3f Normal = Mesh.GetVertexNormal(VID);
+				Us->GetValue(VID, U);
+				float VSign;
+				VSigns->GetValue(VID, &VSign);
+				V = Normal.Cross(U) * FMathf::SignNonZero(VSign);
+			}
+		}
 	}
-};
+}
+
 
 struct FCellMeshes
 {
 	struct FCellInfo
 	{
-		FAugmentedDynamicMesh AugMesh;
+		FDynamicMesh3 AugMesh;
+
+		FCellInfo()
+		{
+			UE::PlanarCutInternals::AugmentDynamicMesh::Augment(AugMesh);
+		}
 
 		// TODO: compute spatial in advance?  (only useful if we rework mesh booleans to support it)
 		//FDynamicMeshAABBTree3 Spatial;
@@ -332,21 +345,21 @@ struct FCellMeshes
 	}
 
 	/**
-	 * Convert plane index to triangle group
-	 * @return group id encoding the source plane into a triangle mesh
+	 * Convert plane index to material ID
+	 * @return material ID encoding the source plane into a triangle mesh
 	 */
-	int PlaneToGroup(int Plane)
+	int PlaneToMaterial(int Plane)
 	{
-		return Plane + 1;
+		return -(Plane+1);
 	}
 
 	/**
-	 * Convert group index to plane index
+	 * Convert material ID to plane index
 	 * @return index of source plane for triangle, or -1 if no such plane
 	 */
-	int GroupToPlane(int Group)
+	int MaterialToPlane(int MaterialID)
 	{
-		return Group - 1;
+		return MaterialID >= 0 ? -1 : -(MaterialID+1);
 	}
 
 	void Init(const FPlanarCells& Cells, FAxisAlignedBox3d DomainBounds, double ExtendDomain, bool bIncludeOutsideCell = true)
@@ -409,17 +422,14 @@ struct FCellMeshes
 		// TODO: self-union on cells when it makes sense to do so (for non-single-plane inputs w/ high noise or possible untracked adjacencies)
 		/*for (FCellInfo& CellInfo : CellMeshes)
 		{
-			FMeshSelfUnion SelfUnion(&CellInfo.AugMesh.Mesh);
+			FMeshSelfUnion SelfUnion(&CellInfo.AugMesh);
 			// TODO: need to have an option in SelfUnion to not weld edges
 			SelfUnion.Compute();
 		}*/
 
 		for (FCellInfo& CellInfo : CellMeshes)
 		{
-			check(CellInfo.AugMesh.Mesh.Attributes() != nullptr);
-			check(CellInfo.AugMesh.Mesh.Attributes()->NumAttachedAttributes() == 3);
-			CellInfo.AugMesh.SetDefaultAttributes(Cells.InternalSurfaceMaterials);
-			CellInfo.AugMesh.SetDefaultMaterials(Cells.InternalSurfaceMaterials.GlobalMaterialID);
+			UE::PlanarCutInternals::AugmentDynamicMesh::SetDefaultAttributes(CellInfo.AugMesh, Cells.InternalSurfaceMaterials);
 		}
 	}
 
@@ -447,12 +457,12 @@ private:
 		for (int32 PlaneIdx = 0; PlaneIdx < Cells.PlaneCells.Num(); PlaneIdx++)
 		{
 			const TPair<int32, int32>& CellPair = Cells.PlaneCells[PlaneIdx];
-			FDynamicMesh3* Meshes[2]{ &CellMeshes[CellPair.Key].AugMesh.Mesh, nullptr };
+			FDynamicMesh3* Meshes[2]{ &CellMeshes[CellPair.Key].AugMesh, nullptr };
 			int32 OtherCell = CellPair.Value < 0 ? OutsideCellIndex : CellPair.Value;
 			int NumMeshes = OtherCell < 0 ? 1 : 2;
 			if (NumMeshes == 2)
 			{
-				Meshes[1] = &CellMeshes[OtherCell].AugMesh.Mesh;
+				Meshes[1] = &CellMeshes[OtherCell].AugMesh;
 			}
 
 			const TArray<int>& PlaneBoundary = Cells.PlaneBoundaries[PlaneIdx];
@@ -498,7 +508,8 @@ private:
 					{
 						Swap(Tri.B, Tri.C);
 					}
-					Meshes[MeshIdx]->AppendTriangle(Tri, PlaneToGroup(PlaneIdx));
+					int TID = Meshes[MeshIdx]->AppendTriangle(Tri);
+					Meshes[MeshIdx]->Attributes()->GetMaterialID()->SetNewValue(TID, PlaneToMaterial(PlaneIdx));
 				}
 			}
 		}
@@ -525,10 +536,11 @@ private:
 		PlaneMeshes.SetNum(Cells.Planes.Num());
 		for (FDynamicMesh3& PlaneMesh : PlaneMeshes)
 		{
-			PlaneMesh.EnableTriangleGroups();
 			PlaneMesh.EnableVertexUVs(FVector2f(0, 0));
 			PlaneMesh.EnableVertexNormals(FVector3f::UnitZ());
 			PlaneMesh.EnableVertexColors(FVector3f(1, 1, 1));
+			PlaneMesh.EnableAttributes();
+			PlaneMesh.Attributes()->EnableMaterialID();
 		}
 
 		struct FPlaneIdxAndFlip
@@ -608,10 +620,11 @@ private:
 				Triangulation.Add(GeneralPolygon);
 				Triangulation.Triangulate();
 
-				int GID = PlaneToGroup(PlaneIdx);
+				int MID = PlaneToMaterial(PlaneIdx);
 				for (FIndex3i Triangle : Triangulation.Triangles)
 				{
-					Mesh.AppendTriangle(Triangle, GID);
+					int TID = Mesh.AppendTriangle(Triangle);
+					Mesh.Attributes()->GetMaterialID()->SetNewValue(TID, MID);
 				}
 
 				RemeshForNoise(Mesh, EEdgeRefineFlags::SplitsOnly, Spacing);
@@ -622,7 +635,7 @@ private:
 		for (int CellIdx = 0; CellIdx < NumCells; CellIdx++)
 		{
 			FCellInfo& CellInfo = CellMeshes[CellIdx];
-			FDynamicMesh3& Mesh = CellInfo.AugMesh.Mesh;
+			FDynamicMesh3& Mesh = CellInfo.AugMesh;
 			bool bFlipForOutsideCell = CellIdx == OutsideCellIndex; // outside cell will be subtracted, and needs all planes flipped vs normal
 			for (FPlaneIdxAndFlip PlaneInfo : CellPlanes[CellIdx])
 			{
@@ -632,7 +645,7 @@ private:
 	}
 	void CreateMeshesForSinglePlane(const FPlanarCells& Cells, const FAxisAlignedBox3d& DomainBounds, bool bNoise, double GlobalUVScale)
 	{
-		int GID = PlaneToGroup(0);
+		int MID = PlaneToMaterial(0);
 		FPlane Plane = Cells.Planes[0];
 
 		FFrame3d PlaneFrame(Plane);
@@ -649,7 +662,9 @@ private:
 		//	// TODO: early out for plane that doesn't even intersect the domain bounding box?
 		//}
 
-		FDynamicMesh3 PlaneMesh(true, true, true, true);
+		FDynamicMesh3 PlaneMesh(true, true, true, false);
+		PlaneMesh.EnableAttributes();
+		PlaneMesh.Attributes()->EnableMaterialID();
 		FVertexInfo PlaneVertInfo;
 		PlaneVertInfo.bHaveC = true;
 		PlaneVertInfo.bHaveUV = true;
@@ -663,8 +678,8 @@ private:
 			PlaneVertInfo.UV = FVector2f(XYRange.GetCorner(CornerIdx) - XYRange.Min) * GlobalUVScale;
 			PlaneMesh.AppendVertex(PlaneVertInfo);
 		}
-		PlaneMesh.AppendTriangle(0, 1, 2, GID);
-		PlaneMesh.AppendTriangle(0, 2, 3, GID);
+		PlaneMesh.AppendTriangle(0, 1, 2);
+		PlaneMesh.AppendTriangle(0, 2, 3);
 
 		if (bNoise)
 		{
@@ -673,7 +688,7 @@ private:
 			ApplyNoise(PlaneMesh, PlaneFrame.GetAxis(2), Cells.InternalSurfaceMaterials.NoiseSettings.GetValue(), true);
 			FMeshNormals::QuickComputeVertexNormals(PlaneMesh);
 		}
-		FDynamicMesh3* Meshes[2]{ &CellMeshes[0].AugMesh.Mesh, &CellMeshes[1].AugMesh.Mesh };
+		FDynamicMesh3* Meshes[2]{ &CellMeshes[0].AugMesh, &CellMeshes[1].AugMesh };
 		for (int Side = 0; Side < 2; Side++)
 		{
 			*Meshes[Side] = PlaneMesh;
@@ -718,9 +733,11 @@ private:
 				MatchedIndices[1].Add(CornerIdx);
 			}
 			int NewTris[2]{
-				Meshes[Side]->AppendTriangle(VertIDs[1][0], VertIDs[1][1], VertIDs[1][2], GID),
-				Meshes[Side]->AppendTriangle(VertIDs[1][0], VertIDs[1][2], VertIDs[1][3], GID)
+				Meshes[Side]->AppendTriangle(VertIDs[1][0], VertIDs[1][1], VertIDs[1][2]),
+				Meshes[Side]->AppendTriangle(VertIDs[1][0], VertIDs[1][2], VertIDs[1][3])
 			};
+			Meshes[Side]->Attributes()->GetMaterialID()->SetNewValue(NewTris[0], MID);
+			Meshes[Side]->Attributes()->GetMaterialID()->SetNewValue(NewTris[1], MID);
 			if (Side == 1)
 			{
 				Meshes[Side]->ReverseTriOrientation(NewTris[0]);
@@ -731,7 +748,7 @@ private:
 			Editor.StitchSparselyCorrespondedVertexLoops(VertIDs[0], MatchedIndices[0], VertIDs[1], MatchedIndices[1], ResultOut, Side == 0);
 
 			// re-enable tangents and visibility attributes, since these are lost when we set the mesh to a copy of the plane mesh
-			CellMeshes[Side].AugMesh.Augment();
+			UE::PlanarCutInternals::AugmentDynamicMesh::Augment(CellMeshes[Side].AugMesh);
 		}
 	}
 };
@@ -743,12 +760,20 @@ struct FDynamicMeshCollection
 {
 	struct FMeshData
 	{
-		FAugmentedDynamicMesh AugMesh;
+		FDynamicMesh3 AugMesh;
 		
 		// FDynamicMeshAABBTree3 Spatial; // TODO: maybe refactor mesh booleans to allow version where caller provides spatial data; it's computed every boolean now
 		// FTransform3d Transform; // TODO: maybe pretransform the data to a space that is good for cutting; refactor mesh boolean so there is an option to have it not transform input
 		int32 TransformIndex; // where the mesh was from in the geometry collection
 		FTransform ToCollection; // transform that need be applied to go back to the local space of the geometry collection
+
+		FMeshData()
+		{
+			UE::PlanarCutInternals::AugmentDynamicMesh::Augment(AugMesh);
+		}
+
+		FMeshData(const FDynamicMesh3& Mesh, int32 TransformIndex, FTransform ToCollection) : AugMesh(Mesh), TransformIndex(TransformIndex), ToCollection(ToCollection)
+		{}
 	};
 	TArray<FMeshData> Meshes;
 	FAxisAlignedBox3d Bounds;
@@ -778,7 +803,7 @@ struct FDynamicMeshCollection
 			FMeshData& MeshData = Meshes.Emplace_GetRef();
 			MeshData.TransformIndex = TransformIdx;
 			MeshData.ToCollection = FTransform(CollectionToLocal.Inverse());
-			FDynamicMesh3& Mesh = MeshData.AugMesh.Mesh;
+			FDynamicMesh3& Mesh = MeshData.AugMesh;
 
 			int32 GeometryIdx = Collection->TransformToGeometryIndex[TransformIdx];
 			Mesh.EnableAttributes();
@@ -799,14 +824,14 @@ struct FDynamicMeshCollection
 				VertexInfo.Color = FVector3f(Collection->Color[Idx]);
 				VertexInfo.Normal = CollectionToLocal.TransformNormal(FVector3f(Collection->Normal[Idx]));
 				int VID = Mesh.AppendVertex(VertexInfo);
-				MeshData.AugMesh.SetTangent(VID, VertexInfo.Normal, Collection->TangentU[Idx], Collection->TangentV[Idx]);
+				UE::PlanarCutInternals::AugmentDynamicMesh::SetTangent(Mesh, VID, VertexInfo.Normal, Collection->TangentU[Idx], Collection->TangentV[Idx]);
 			}
 			FIntVector VertexOffset(VertexStart, VertexStart, VertexStart);
 			for (int32 Idx = Collection->FaceStart[GeometryIdx], N = Collection->FaceStart[GeometryIdx] + FaceCount; Idx < N; Idx++)
 			{
 				int TID = Mesh.AppendTriangle(FIndex3i(Collection->Indices[Idx] - VertexOffset), 0);
 				Mesh.Attributes()->GetMaterialID()->SetValue(TID, Collection->MaterialID[Idx]);
-				MeshData.AugMesh.SetVisibility(TID, Collection->Visible[Idx]);
+				UE::PlanarCutInternals::AugmentDynamicMesh::SetVisibility(Mesh, TID, Collection->Visible[Idx]);
 				// note: material index doesn't need to be passed through; will be rebuilt by a call to reindex materials once the cut mesh is returned back to geometry collection format
 			}
 
@@ -849,7 +874,7 @@ struct FDynamicMeshCollection
 		TArray<TUniquePtr<TPointHashGrid3d<int>>> VerticesHashes;
 		auto HashMeshVertices = [&VerticesHashes, &ToCut](int32 HashIdx)
 		{
-			FDynamicMesh3& Mesh = ToCut[HashIdx]->AugMesh.Mesh;
+			FDynamicMesh3& Mesh = ToCut[HashIdx]->AugMesh;
 			if (HashIdx >= VerticesHashes.Num())
 			{
 				VerticesHashes.SetNum(HashIdx + 1);
@@ -882,16 +907,16 @@ struct FDynamicMeshCollection
 			{
 				return false;
 			}
-			if (!ToCut[A]->AugMesh.Mesh.GetCachedBounds().Intersects(ToCut[B]->AugMesh.Mesh.GetCachedBounds()))
+			if (!ToCut[A]->AugMesh.GetCachedBounds().Intersects(ToCut[B]->AugMesh.GetCachedBounds()))
 			{
 				return false;
 			}
-			if (ToCut[A]->AugMesh.Mesh.VertexCount() > ToCut[B]->AugMesh.Mesh.VertexCount())
+			if (ToCut[A]->AugMesh.VertexCount() > ToCut[B]->AugMesh.VertexCount())
 			{
 				Swap(A, B);
 			}
-			FDynamicMesh3& RefMesh = ToCut[B]->AugMesh.Mesh;
-			for (const FVector3d& V : ToCut[A]->AugMesh.Mesh.VerticesItr())
+			FDynamicMesh3& RefMesh = ToCut[B]->AugMesh;
+			for (const FVector3d& V : ToCut[A]->AugMesh.VerticesItr())
 			{
 				TPair<int, double> Nearest = VerticesHashes[B]->FindNearestInRadius(V, FMathd::ZeroTolerance * 10, [&RefMesh, &V](int VID)
 					{
@@ -932,8 +957,10 @@ struct FDynamicMeshCollection
 			for (int32 ToCutIdx = 0, ToCutNum = ToCut.Num(); ToCutIdx < ToCutNum; ToCutIdx++)
 			{
 				FMeshData& Surface = *ToCut[ToCutIdx];
+				int32 TransformIndex = Surface.TransformIndex;
+				FTransform ToCollection = Surface.ToCollection;
 
-				FAxisAlignedBox3d Box = Surface.AugMesh.Mesh.GetCachedBounds();
+				FAxisAlignedBox3d Box = Surface.AugMesh.GetCachedBounds();
 				if (InternalSurfaceMaterials.NoiseSettings.IsSet())
 				{
 					Box.Expand(InternalSurfaceMaterials.NoiseSettings->Amplitude);
@@ -947,8 +974,8 @@ struct FDynamicMeshCollection
 				for (int ResultIdx = 0; ResultIdx < 2; ResultIdx++)
 				{
 					BoolResults.Add(MakeUnique<FMeshData>());
-					BoolResults[ResultIdx]->TransformIndex = Surface.TransformIndex;
-					BoolResults[ResultIdx]->ToCollection = Surface.ToCollection;
+					BoolResults[ResultIdx]->TransformIndex = TransformIndex;
+					BoolResults[ResultIdx]->ToCollection = ToCollection;
 				}
 				check(CellMeshes.CellMeshes.Num() == 2);
 				bool bKeepResults = true;
@@ -961,7 +988,7 @@ struct FDynamicMeshCollection
 					{
 						Op = FMeshBoolean::EBooleanOp::Difference;
 					}
-					FMeshBoolean Boolean(&Surface.AugMesh.Mesh, &Cell.AugMesh.Mesh, &BoolResults[CellIdx]->AugMesh.Mesh, Op);
+					FMeshBoolean Boolean(&Surface.AugMesh, &Cell.AugMesh, &BoolResults[CellIdx]->AugMesh, Op);
 					Boolean.bWeldSharedEdges = false;
 					if (!Boolean.Compute())
 					{
@@ -969,7 +996,7 @@ struct FDynamicMeshCollection
 						// note: failure cases won't be detected at all unless we weld edges,
 						//       which will require re-working how tangents are carried through
 					}
-					if (BoolResults[CellIdx]->AugMesh.Mesh.TriangleCount() == 0)
+					if (BoolResults[CellIdx]->AugMesh.TriangleCount() == 0)
 					{
 						bKeepResults = false;
 						break;
@@ -980,7 +1007,24 @@ struct FDynamicMeshCollection
 				{
 					ToCut[ToCutIdx] = MoveTemp(BoolResults[0]);
 					int32 NewIdx = ToCut.Add(MoveTemp(BoolResults[1]));
-					int32 ResultIndices[2] = { ToCutIdx, NewIdx };
+					// indices of all boolean result meshes (may be more than two due to splitting disconnected components)
+					TArray<int32, TInlineAllocator<4>> ResultIndices = { ToCutIdx, NewIdx };
+					// corresponding parent indices for each result mesh
+					TArray<int32, TInlineAllocator<4>> ParentIndices = { 0, 1 };
+					TArray<FDynamicMesh3> SplitMeshes;
+					for (int UnsplitIdx = 0; UnsplitIdx < 2; UnsplitIdx++)
+					{
+						if (SplitIslands(ToCut[ResultIndices[UnsplitIdx]]->AugMesh, SplitMeshes))
+						{
+							ToCut[ResultIndices[UnsplitIdx]]->AugMesh = SplitMeshes[0];
+							for (int32 Idx = 1; Idx < SplitMeshes.Num(); Idx++)
+							{
+								const FDynamicMesh3& Mesh = SplitMeshes[Idx];
+								ResultIndices.Add(ToCut.Add(MakeUnique<FMeshData>(Mesh, TransformIndex, ToCollection)));
+								ParentIndices.Add(UnsplitIdx);
+							}
+						}
+					}
 
 					// update proximity for neighbors of the original piece
 					if (bHasProximity)
@@ -990,14 +1034,16 @@ struct FDynamicMeshCollection
 						Proximity.MultiFind(ToCutIdx, Nbrs);
 						if (Nbrs.Num() > 0)
 						{
-							HashMeshVertices(ToCutIdx);
-							HashMeshVertices(NewIdx);
+							for (int32 ChangedMeshIdx : ResultIndices)
+							{
+								HashMeshVertices(ChangedMeshIdx);
+							}
 
 							for (int32 Nbr : Nbrs)
 							{
 								ProxUnlink(ToCutIdx, Nbr);
 								HashMeshVertices(Nbr);
-								for (int32 Idx = 0; Idx < 2; Idx++)
+								for (int32 Idx = 0; Idx < ResultIndices.Num(); Idx++)
 								{
 									int32 ResultIdx = ResultIndices[Idx];
 									int32 OldIdx = Nbr;
@@ -1008,12 +1054,42 @@ struct FDynamicMeshCollection
 								}
 							}
 						}
-					}
 
-					// add the connection between the two new pieces
-					ProxLink(ResultIndices[0], ResultIndices[1]);
-				}
-			}
+						if (ResultIndices.Num() == 2)
+						{
+							// add the connection between the two new pieces
+							ProxLink(ResultIndices[0], ResultIndices[1]);
+						}
+						else
+						{
+							if (Nbrs.Num() == 0)
+							{
+								for (int32 ChangedMeshIdx : ResultIndices)
+								{
+									HashMeshVertices(ChangedMeshIdx);
+								}
+							}
+							// check for connections between all pieces
+							for (int FirstIdx = 0; FirstIdx + 1 < ResultIndices.Num(); FirstIdx++)
+							{
+								int32 FirstParent = ParentIndices[FirstIdx];
+								for (int SecondIdx = FirstIdx + 1; SecondIdx < ResultIndices.Num(); SecondIdx++)
+								{
+									if (FirstParent == ParentIndices[SecondIdx])
+									{
+										// these pieces split from the same mesh *because* they were disconnected, so the pieces cannot be neighbors
+										continue;
+									}
+									if (IsNeighbor(ResultIndices[FirstIdx], ResultIndices[SecondIdx]))
+									{
+										ProxLink(ResultIndices[FirstIdx], ResultIndices[SecondIdx]);
+									}
+								}
+							}
+						}
+					}
+				} // iteration over meshes to cut
+			} // iteration over cutting planes
 		}
 
 		TMultiMap<int32, int32> ParentTransformToChildren;
@@ -1028,7 +1104,7 @@ struct FDynamicMeshCollection
 		for (FMeshData& MeshData : Meshes)
 		{
 			int32 GeometryIdx = Collection->TransformToGeometryIndex[MeshData.TransformIndex];
-			int32 MaterialIDOverride = bSetDefaultInternalMaterialsFromCollection ? InternalSurfaceMaterials.GetDefaultMaterialIDForGeometry(*Collection, GeometryIdx) : -1;
+			int32 InternalMaterialID = bSetDefaultInternalMaterialsFromCollection ? InternalSurfaceMaterials.GetDefaultMaterialIDForGeometry(*Collection, GeometryIdx) : InternalSurfaceMaterials.GlobalMaterialID;
 			ParentTransformToChildren.MultiFind(MeshData.TransformIndex, ToCutIndices);
 
 			// if there's only one mesh here, i.e. it didn't get cut at all
@@ -1044,19 +1120,10 @@ struct FDynamicMeshCollection
 			int32 SubPartIdx = 0;
 			for (int32 ToCutIdx : ToCutIndices)
 			{
-				ToCut[ToCutIdx]->AugMesh.FixTangents();
-				for (int TID : ToCut[ToCutIdx]->AugMesh.Mesh.TriangleIndicesItr())
-				{
-					int GID = ToCut[ToCutIdx]->AugMesh.Mesh.GetTriangleGroup(TID);
-					if (GID > 0)
-					{
-						if (MaterialIDOverride != -1)
-						{
-							ToCut[ToCutIdx]->AugMesh.Mesh.Attributes()->GetMaterialID()->SetNewValue(TID, MaterialIDOverride);
-						}
-					}
-				}
-				int32 CreatedGeometryIdx = AppendToCollection(ToCut[ToCutIdx]->ToCollection, ToCut[ToCutIdx]->AugMesh, ToCut[ToCutIdx]->TransformIndex, SubPartIdx++, *Collection);
+				FDynamicMesh3& Mesh = ToCut[ToCutIdx]->AugMesh;
+				UE::PlanarCutInternals::AugmentDynamicMesh::FixTangents(Mesh);
+
+				int32 CreatedGeometryIdx = AppendToCollection(ToCut[ToCutIdx]->ToCollection, Mesh, ToCut[ToCutIdx]->TransformIndex, SubPartIdx++, *Collection, InternalMaterialID);
 				ToCutIdxToGeometryIdx[ToCutIdx] = CreatedGeometryIdx;
 				if (FirstCreatedIndex == -1)
 				{
@@ -1096,20 +1163,21 @@ struct FDynamicMeshCollection
 		for (FMeshData& Surface : Meshes)
 		{
 			int32 GeometryIdx = Collection->TransformToGeometryIndex[Surface.TransformIndex];
-			TArray<FAugmentedDynamicMesh> BooleanResults;  BooleanResults.SetNum(CellMeshes.CellMeshes.Num());
+			TArray<TUniquePtr<FDynamicMesh3>> BooleanResults;  BooleanResults.SetNum(CellMeshes.CellMeshes.Num());
 			ParallelFor(CellMeshes.CellMeshes.Num(), [&BooleanResults, &CellMeshes, &Surface](int32 CellIdx)
 				{
 					FCellMeshes::FCellInfo& Cell = CellMeshes.CellMeshes[CellIdx];
-					if (Cell.AugMesh.Mesh.GetCachedBounds().Intersects(Surface.AugMesh.Mesh.GetCachedBounds()))
+					if (Cell.AugMesh.GetCachedBounds().Intersects(Surface.AugMesh.GetCachedBounds()))
 					{
-						FAugmentedDynamicMesh& AugBoolResult = BooleanResults[CellIdx];
+						BooleanResults[CellIdx] = MakeUnique<FDynamicMesh3>();
+						FDynamicMesh3& AugBoolResult = *BooleanResults[CellIdx];
 
 						FMeshBoolean::EBooleanOp Op = FMeshBoolean::EBooleanOp::Intersect;
 						if (CellIdx == CellMeshes.OutsideCellIndex)
 						{
 							Op = FMeshBoolean::EBooleanOp::Difference;
 						}
-						FMeshBoolean Boolean(&Surface.AugMesh.Mesh, &Cell.AugMesh.Mesh, &AugBoolResult.Mesh, Op);
+						FMeshBoolean Boolean(&Surface.AugMesh, &Cell.AugMesh, &AugBoolResult, Op);
 						Boolean.bWeldSharedEdges = false;
 						if (!Boolean.Compute())
 						{
@@ -1121,9 +1189,9 @@ struct FDynamicMeshCollection
 				}, EParallelForFlags::None);
 
 			int32 NonEmptyResults = 0;
-			for (const FAugmentedDynamicMesh& AugBoolResult : BooleanResults)
+			for (const TUniquePtr<FDynamicMesh3>& AugBoolResult : BooleanResults)
 			{
-				if (AugBoolResult.Mesh.TriangleCount() > 0)
+				if (AugBoolResult.IsValid() && AugBoolResult->TriangleCount() > 0)
 				{
 					NonEmptyResults++;
 				}
@@ -1132,54 +1200,105 @@ struct FDynamicMeshCollection
 			if (NonEmptyResults > 1) // only write to geometry collection if more than one result was non-empty
 			{
 				TSet<int32> PlanesInOutput;
-				TMap<int32, int32> CellToGeometry;
+				TMultiMap<int32, int32> CellToGeometry;
+				TMap<int32, int32> GeometryToResultMesh;
 				int32 SubPartIndex = 0;
-				int32 MaterialIDOverride = bSetDefaultInternalMaterialsFromCollection ? PlanarCells.InternalSurfaceMaterials.GetDefaultMaterialIDForGeometry(*Collection, GeometryIdx) : -1;
+				int32 InternalMaterialID = bSetDefaultInternalMaterialsFromCollection ? PlanarCells.InternalSurfaceMaterials.GetDefaultMaterialIDForGeometry(*Collection, GeometryIdx) : PlanarCells.InternalSurfaceMaterials.GlobalMaterialID;
 
 				for (int32 CellIdx = 0; CellIdx < CellMeshes.CellMeshes.Num(); CellIdx++)
-				{
-					FAugmentedDynamicMesh& AugBoolResult = BooleanResults[CellIdx];
-					if (AugBoolResult.Mesh.TriangleCount() > 0)
+				{					
+					if (BooleanResults[CellIdx].IsValid() && BooleanResults[CellIdx]->TriangleCount() > 0)
 					{
-						AugBoolResult.FixTangents();
-						for (int TID : AugBoolResult.Mesh.TriangleIndicesItr())
+						FDynamicMesh3& AugBoolResult = *BooleanResults[CellIdx];
+						UE::PlanarCutInternals::AugmentDynamicMesh::FixTangents(AugBoolResult);
+						for (int TID : AugBoolResult.TriangleIndicesItr())
 						{
-							int GID = AugBoolResult.Mesh.GetTriangleGroup(TID);
-							int32 PlaneIdx = CellMeshes.GroupToPlane(GID);
+							int MID = AugBoolResult.Attributes()->GetMaterialID()->GetValue(TID);
+							int32 PlaneIdx = CellMeshes.MaterialToPlane(MID);
 							if (PlaneIdx >= 0)
 							{
 								PlanesInOutput.Add(PlaneIdx);
-								if (MaterialIDOverride != -1)
+							}
+						}
+						int32 CreatedGeometryIdx = -1;
+						TArray<FDynamicMesh3> Islands;
+						if (SplitIslands(AugBoolResult, Islands))
+						{
+							for (int32 i = 0; i < Islands.Num(); i++)
+							{
+								FDynamicMesh3& Island = Islands[i];
+								CreatedGeometryIdx = AppendToCollection(Surface.ToCollection, Island, Surface.TransformIndex, SubPartIndex++, *Collection, InternalMaterialID);
+								CellToGeometry.Add(CellIdx, CreatedGeometryIdx);
+								if (i > 0)
+								{
+									GeometryToResultMesh.Add(CreatedGeometryIdx, BooleanResults.Add(MakeUnique<FDynamicMesh3>(Island)));
+								}
+								else
 								{
-									AugBoolResult.Mesh.Attributes()->GetMaterialID()->SetNewValue(TID, MaterialIDOverride);
+									*BooleanResults[CellIdx] = Island;
+									GeometryToResultMesh.Add(CreatedGeometryIdx, CellIdx);
 								}
 							}
 						}
-						int32 CreatedGeometryIdx = AppendToCollection(Surface.ToCollection, AugBoolResult, Surface.TransformIndex, SubPartIndex, *Collection);
-						CellToGeometry.Add(CellIdx, CreatedGeometryIdx);
+						else
+						{
+							CreatedGeometryIdx = AppendToCollection(Surface.ToCollection, AugBoolResult, Surface.TransformIndex, SubPartIndex++, *Collection, InternalMaterialID);
+							CellToGeometry.Add(CellIdx, CreatedGeometryIdx);
+							GeometryToResultMesh.Add(CreatedGeometryIdx, CellIdx);
+						}
 						if (FirstIdx == -1)
 						{
 							FirstIdx = CreatedGeometryIdx;
 						}
-						SubPartIndex++;
 					}
-					AugBoolResult.Mesh.Clear(); // we can stop storing the dynamic mesh version now that it's in the geometry collection
 				}
 				if (bHasProximity)
 				{
 					TManagedArray<TSet<int32>>& Proximity = Collection->GetAttribute<TSet<int32>>("Proximity", FGeometryCollection::GeometryGroup);
+					TArray<TUniquePtr<TPointHashGrid3d<int>>> VertexHashes;
+					auto MakeHash = [this, &VertexHashes, &BooleanResults](int GID)
+					{
+						if (GID >= VertexHashes.Num())
+						{
+							VertexHashes.SetNum(GID + 1);
+						}
+						if (!VertexHashes[GID].IsValid())
+						{
+							VertexHashes[GID] = MakeUnique<TPointHashGrid3d<int>>(FMathd::ZeroTolerance * 1000, -1);
+							FillVertexHash(*BooleanResults[GID], *VertexHashes[GID]);
+						}
+					};
 					for (int32 PlaneIdx : PlanesInOutput)
 					{
 						TPair<int32, int32> Cells = PlanarCells.PlaneCells[PlaneIdx];
-						int32 SecondCell = Cells.Value == -1 ? CellMeshes.OutsideCellIndex : Cells.Value;
+						int32 SecondCell = Cells.Value < 0 ? CellMeshes.OutsideCellIndex : Cells.Value;
 						if (SecondCell != -1)
 						{
-							int32* GeomA = CellToGeometry.Find(Cells.Key);
-							int32* GeomB = CellToGeometry.Find(SecondCell);
-							if (GeomA && GeomB)
+							TArray<int32, TInlineAllocator<4>> GeomA, GeomB;
+							CellToGeometry.MultiFind(Cells.Key, GeomA, false);
+							CellToGeometry.MultiFind(SecondCell, GeomB, false);
+							if (GeomA.Num() == 1 && GeomB.Num() == 1)
 							{
-								Proximity[*GeomA].Add(*GeomB);
-								Proximity[*GeomB].Add(*GeomA);
+								Proximity[GeomA[0]].Add(GeomB[0]);
+								Proximity[GeomB[0]].Add(GeomA[0]);
+							}
+							else if (GeomA.Num() >= 1 && GeomB.Num() >= 1) // at least one was split; need to re-check proximities
+							{
+								for (int GIDA : GeomA)
+								{
+									int MeshA = GeometryToResultMesh[GIDA];
+									MakeHash(MeshA);
+									for (int GIDB : GeomB)
+									{
+										int MeshB = GeometryToResultMesh[GIDB];
+										MakeHash(MeshB);
+										if (IsNeighboring(*BooleanResults[MeshA], *VertexHashes[MeshA], *BooleanResults[MeshB], *VertexHashes[MeshB]))
+										{
+											Proximity[GIDA].Add(GIDB);
+											Proximity[GIDB].Add(GIDA);
+										}
+									}
+								}
 							}
 						}
 					}
@@ -1201,9 +1320,87 @@ struct FDynamicMeshCollection
 		}
 	}
 
-	static int32 AppendToCollection(const FTransform& ToCollection, FAugmentedDynamicMesh& AugMesh, int32 TransformParent, int32 SubPartIndex, FGeometryCollection& Output)
+	void FillVertexHash(const FDynamicMesh3& Mesh, TPointHashGrid3d<int>& VertHash)
+	{
+		for (int VID : Mesh.VertexIndicesItr())
+		{
+			FVector3d V = Mesh.GetVertex(VID);
+			VertHash.InsertPointUnsafe(VID, V);
+		}
+	}
+
+	bool IsNeighboring(FDynamicMesh3& MeshA, const TPointHashGrid3d<int>& VertHashA, FDynamicMesh3& MeshB, const TPointHashGrid3d<int>& VertHashB)
+	{
+		FDynamicMesh3* Mesh[2]{ &MeshA, &MeshB };
+		const TPointHashGrid3d<int>* VertHash[2]{ &VertHashA, &VertHashB };
+		return IsNeighboring(Mesh, VertHash);
+	}
+
+	bool IsNeighboring(FDynamicMesh3* Mesh[2], const TPointHashGrid3d<int>* VertHash[2])
+	{
+		if (!ensure(Mesh[0] && Mesh[1] && VertHash[0] && VertHash[1]))
+		{
+			return false;
+		}
+		if (!Mesh[0]->GetCachedBounds().Intersects(Mesh[1]->GetCachedBounds()))
+		{
+			return false;
+		}
+		
+		int A = 0, B = 1;
+		if (Mesh[0]->VertexCount() > Mesh[1]->VertexCount())
+		{
+			Swap(A, B);
+		}
+		FDynamicMesh3& RefMesh = *Mesh[B];
+		for (const FVector3d& V : Mesh[A]->VerticesItr())
+		{
+			TPair<int, double> Nearest = VertHash[B]->FindNearestInRadius(V, FMathd::ZeroTolerance * 10, [&RefMesh, &V](int VID)
+				{
+					return RefMesh.GetVertex(VID).DistanceSquared(V);
+				});
+			if (Nearest.Key != -1)
+			{
+				return true;
+			}
+		}
+		return false;
+	}
+
+	// Split mesh into connected components, including implicit connections by co-located vertices
+	bool SplitIslands(FDynamicMesh3& Source, TArray<FDynamicMesh3>& SeparatedMeshes)
+	{
+		double SnapDistance = FMathd::ZeroTolerance * 100;
+		TPointHashGrid3d<int> VertHash(SnapDistance * 10, -1);
+		FDisjointSet VertComponents(Source.MaxVertexID());
+		// Add Source vertices to hash & disjoint sets
+		TArray<int> Neighbors;
+		for (int VID : Source.VertexIndicesItr())
+		{
+			FVector3d Pt = Source.GetVertex(VID);
+			Neighbors.Reset();
+			VertHash.FindPointsInBall(Pt, SnapDistance, [&Source, Pt](int OtherVID) {return Pt.DistanceSquared(Source.GetVertex(OtherVID));}, Neighbors);
+			for (int NbrVID : Neighbors)
+			{
+				VertComponents.UnionSequential(VID, NbrVID);
+			}
+			VertHash.InsertPointUnsafe(VID, Pt);
+		}
+		for (FIndex3i Tri : Source.TrianglesItr())
+		{
+			VertComponents.Union(Tri.A, Tri.B);
+			VertComponents.Union(Tri.B, Tri.C);
+			VertComponents.Union(Tri.C, Tri.A);
+		}
+		
+		return FDynamicMeshEditor::SplitMesh(&Source, SeparatedMeshes, [&Source, &VertComponents](int TID)
+			{
+				return (int)VertComponents.Find(Source.GetTriangle(TID).A);
+			});
+	}
+
+	static int32 AppendToCollection(const FTransform& ToCollection, FDynamicMesh3& Mesh, int32 TransformParent, int32 SubPartIndex, FGeometryCollection& Output, int32 InternalMaterialID)
 	{
-		FDynamicMesh3& Mesh = AugMesh.Mesh;
 		if (Mesh.TriangleCount() == 0)
 		{
 			return -1;
@@ -1252,7 +1449,7 @@ struct FDynamicMeshCollection
 			Output.Normal[CopyToIdx] = ToCollection.TransformVectorNoScale(FVector(Mesh.GetVertexNormal(VID))); // TODO: use correct normal transform instead
 			Output.UV[CopyToIdx] = FVector2D(Mesh.GetVertexUV(VID));
 			FVector3f TangentU, TangentV;
-			AugMesh.GetTangent(VID, TangentU, TangentV);
+			UE::PlanarCutInternals::AugmentDynamicMesh::GetTangent(Mesh, VID, TangentU, TangentV);
 			Output.TangentU[CopyToIdx] = ToCollection.TransformVectorNoScale(FVector(TangentU));
 			Output.TangentV[CopyToIdx] = ToCollection.TransformVectorNoScale(FVector(TangentV));
 			Output.Color[CopyToIdx] = FVector(Mesh.GetVertexColor(VID));
@@ -1266,8 +1463,9 @@ struct FDynamicMeshCollection
 		{
 			check(Mesh.IsTriangle(TID));
 			int32 CopyToIdx = FacesStart + TID;
-			Output.Visible[CopyToIdx] = AugMesh.GetVisibility(TID);
-			Output.MaterialID[CopyToIdx] = Mesh.Attributes()->GetMaterialID()->GetValue(TID);
+			Output.Visible[CopyToIdx] = UE::PlanarCutInternals::AugmentDynamicMesh::GetVisibility(Mesh, TID);
+			int MaterialID = Mesh.Attributes()->GetMaterialID()->GetValue(TID);
+			Output.MaterialID[CopyToIdx] = MaterialID < 0 ? InternalMaterialID : MaterialID;
 			Output.Indices[CopyToIdx] = FIntVector(Mesh.GetTriangle(TID)) + VertexStartOffset;
 		}
 
