avoid creating self-intersections when adding noise to multi-cell cuts
+ fix meshing for non-convex cutting cell generation
+ fix tests to pass again w/ updated code

#rnx

diff --git a/Engine/Plugins/Experimental/PlanarCutPlugin/Source/PlanarCut/Private/PlanarCut.cpp b/Engine/Plugins/Experimental/PlanarCutPlugin/Source/PlanarCut/Private/PlanarCut.cpp
--- a/Engine/Plugins/Experimental/PlanarCutPlugin/Source/PlanarCut/Private/PlanarCut.cpp
+++ b/Engine/Plugins/Experimental/PlanarCutPlugin/Source/PlanarCut/Private/PlanarCut.cpp
@@ -498,18 +498,55 @@ private:
 				}
 			}
 
-			for (int V0 = 0, V1 = 1, V2 = 2; V2 < PlaneBoundary.Num(); V1 = V2++)
+			int MID = PlaneToMaterial(PlaneIdx);
+			if (Cells.AssumeConvexCells)
 			{
+				// put a fan
+				for (int V0 = 0, V1 = 1, V2 = 2; V2 < PlaneBoundary.Num(); V1 = V2++)
+				{
+					for (int MeshIdx = 0; MeshIdx < NumMeshes; MeshIdx++)
+					{
+						int Offset = VertStart[MeshIdx];
+						FIndex3i Tri(V0 + Offset, V1 + Offset, V2 + Offset);
+						if (MeshIdx == 1 && OtherCell != OutsideCellIndex)
+						{
+							Swap(Tri.B, Tri.C);
+						}
+						int TID = Meshes[MeshIdx]->AppendTriangle(Tri);
+						Meshes[MeshIdx]->Attributes()->GetMaterialID()->SetNewValue(TID, MID);
+					}
+				}
+			}
+			else // cells may not be convex; cannot triangulate w/ fan
+			{
+				// Delaunay triangulate
+				FPolygon2f Polygon;
+				for (int V = 0; V < PlaneBoundary.Num(); V++)
+				{
+					Polygon.AppendVertex(Meshes[0]->GetVertexUV(VertStart[0] + V));
+				}
+
+				FGeneralPolygon2f GeneralPolygon(Polygon);
+				FConstrainedDelaunay2f Triangulation;
+				Triangulation.FillRule = FConstrainedDelaunay2f::EFillRule::NonZero;
+				Triangulation.Add(GeneralPolygon);
+				Triangulation.Triangulate();
+
 				for (int MeshIdx = 0; MeshIdx < NumMeshes; MeshIdx++)
 				{
 					int Offset = VertStart[MeshIdx];
-					FIndex3i Tri(V0 + Offset, V1 + Offset, V2 + Offset);
-					if (MeshIdx == 1 && OtherCell != OutsideCellIndex)
+					for (FIndex3i Triangle : Triangulation.Triangles)
 					{
-						Swap(Tri.B, Tri.C);
+						Triangle.A += Offset;
+						Triangle.B += Offset;
+						Triangle.C += Offset;
+						if (MeshIdx == 1 && OtherCell != OutsideCellIndex)
+						{
+							Swap(Triangle.B, Triangle.C);
+						}
+						int TID = Meshes[MeshIdx]->AppendTriangle(Triangle);
+						Meshes[MeshIdx]->Attributes()->GetMaterialID()->SetNewValue(TID, MID);
 					}
-					int TID = Meshes[MeshIdx]->AppendTriangle(Tri);
-					Meshes[MeshIdx]->Attributes()->GetMaterialID()->SetNewValue(TID, PlaneToMaterial(PlaneIdx));
 				}
 			}
 		}
@@ -534,6 +571,7 @@ private:
 	{
 		TArray<FDynamicMesh3> PlaneMeshes;
 		PlaneMeshes.SetNum(Cells.Planes.Num());
+		FName OriginalPositionAttribute = "OriginalPosition";
 		for (FDynamicMesh3& PlaneMesh : PlaneMeshes)
 		{
 			PlaneMesh.EnableVertexUVs(FVector2f(0, 0));
@@ -541,6 +579,7 @@ private:
 			PlaneMesh.EnableVertexColors(FVector3f(1, 1, 1));
 			PlaneMesh.EnableAttributes();
 			PlaneMesh.Attributes()->EnableMaterialID();
+			PlaneMesh.Attributes()->AttachAttribute(OriginalPositionAttribute, new TDynamicMeshVertexAttribute<double, 3>(&PlaneMesh));
 		}
 
 		struct FPlaneIdxAndFlip
@@ -581,8 +620,8 @@ private:
 			TotalArea += AreaVec.Size();
 		}
 		double Spacing = GetSafeNoiseSpacing(TotalArea, Cells.InternalSurfaceMaterials.NoiseSettings->PointSpacing);
-
-		ParallelFor(Cells.Planes.Num(), [this, &PlaneMeshes, &Cells, GlobalUVScale, Spacing](int32 PlaneIdx)
+		
+		ParallelFor(Cells.Planes.Num(), [this, OriginalPositionAttribute, &PlaneMeshes, &Cells, GlobalUVScale, Spacing](int32 PlaneIdx)
 			{
 				FDynamicMesh3& Mesh = PlaneMeshes[PlaneIdx];
 				const TArray<int>& PlaneBoundary = Cells.PlaneBoundaries[PlaneIdx];
@@ -628,7 +667,14 @@ private:
 				}
 
 				RemeshForNoise(Mesh, EEdgeRefineFlags::SplitsOnly, Spacing);
+				TDynamicMeshVertexAttribute<double, 3>* OriginalPosns =
+					static_cast<TDynamicMeshVertexAttribute<double, 3>*>(Mesh.Attributes()->GetAttachedAttribute(OriginalPositionAttribute));
+				for (int VID : Mesh.VertexIndicesItr())
+				{
+					OriginalPosns->SetValue(VID, Mesh.GetVertex(VID));
+				}
 				ApplyNoise(Mesh, FVector3d(Normal), Cells.InternalSurfaceMaterials.NoiseSettings.GetValue());
+
 				FMeshNormals::QuickComputeVertexNormals(Mesh);
 			}, EParallelForFlags::None);
 
@@ -636,12 +682,187 @@ private:
 		{
 			FCellInfo& CellInfo = CellMeshes[CellIdx];
 			FDynamicMesh3& Mesh = CellInfo.AugMesh;
+			Mesh.Attributes()->AttachAttribute(OriginalPositionAttribute, new TDynamicMeshVertexAttribute<double, 3>(&Mesh));
 			bool bFlipForOutsideCell = CellIdx == OutsideCellIndex; // outside cell will be subtracted, and needs all planes flipped vs normal
 			for (FPlaneIdxAndFlip PlaneInfo : CellPlanes[CellIdx])
 			{
 				AppendMesh(Mesh, PlaneMeshes[PlaneInfo.PlaneIdx], PlaneInfo.bIsFlipped ^ bFlipForOutsideCell);
 			}
 		}
+
+		// resolve self-intersections
+
+		// build hash grid of mesh vertices so we correspond all same-pos vertices across touching meshes
+		TPointHashGrid3d<FIndex2i> MeshesVertices(FMathd::ZeroTolerance*1000, FIndex2i::Invalid());
+		for (int CellIdx = 0; CellIdx < NumCells; CellIdx++)
+		{
+			FCellInfo& CellInfo = CellMeshes[CellIdx];
+			FDynamicMesh3& Mesh = CellInfo.AugMesh;
+			for (int VID : Mesh.VertexIndicesItr())
+			{
+				MeshesVertices.InsertPointUnsafe(FIndex2i(CellIdx, VID), Mesh.GetVertex(VID));
+			}
+		}
+		
+		// repeatedly detect and resolve collisions until there are no more (or give up after too many iterations)
+		TArray<bool> CellUnmoved; CellUnmoved.Init(false, NumCells);
+		const int MaxIters = 10;
+		for (int Iters = 0; Iters < MaxIters; Iters++)
+		{
+			struct FUpdate
+			{
+				FIndex2i Tris;
+				TArray<FIndex2i> IDs;
+				FUpdate(int TriA = -1, int TriB = -1) : Tris(TriA, TriB)
+				{}
+			};
+
+			// todo: can parallelize?
+			TArray<TArray<FUpdate>> Updates; Updates.SetNum(NumCells);
+			bool bAnyUpdatesNeeded = false;
+			for (int CellIdx = 0; CellIdx < NumCells; CellIdx++)
+			{
+				if (CellUnmoved[CellIdx])
+				{
+					// if nothing moved since last time we resolved self intersections on this cell, don't need to process again
+					continue;
+				}
+				FDynamicMesh3& Mesh = CellMeshes[CellIdx].AugMesh;
+				FDynamicMeshAABBTree3 CellTree(&Mesh, true);
+				MeshIntersection::FIntersectionsQueryResult Intersections = CellTree.FindAllSelfIntersections(true);
+				for (MeshIntersection::FSegmentIntersection& Seg : Intersections.Segments)
+				{
+					// manually check for shared edges by vertex position because they might not be topologically connected
+					FIndex3i Tri[2]{ Mesh.GetTriangle(Seg.TriangleID[0]), Mesh.GetTriangle(Seg.TriangleID[1]) };
+					int MatchedVertices = 0;
+					for (int T0SubIdx = 0; T0SubIdx < 3; T0SubIdx++)
+					{
+						FVector3d V0 = Mesh.GetVertex(Tri[0][T0SubIdx]);
+						for (int T1SubIdx = 0; T1SubIdx < 3; T1SubIdx++)
+						{
+							FVector3d V1 = Mesh.GetVertex(Tri[1][T1SubIdx]);
+							if (V0.DistanceSquared(V1) < FMathd::ZeroTolerance)
+							{
+								MatchedVertices++;
+								break;
+							}
+						}
+					}
+					// no shared vertices: treat as a real collision
+					// (TODO: only skip shared edges? will need to do something to avoid shared vertices becoming collisions)
+					if (MatchedVertices < 1)
+					{
+						bAnyUpdatesNeeded = true;
+						FUpdate& Update = Updates[CellIdx].Emplace_GetRef(Seg.TriangleID[0], Seg.TriangleID[1]);
+						for (int TriIdx = 0; TriIdx < 2; TriIdx++)
+						{
+							for (int VSubIdx = 0; VSubIdx < 3; VSubIdx++)
+							{
+								int VIdx = Tri[TriIdx][VSubIdx];
+								FVector3d P = Mesh.GetVertex(VIdx);
+								FIndex2i IDs(CellIdx, VIdx);
+								MeshesVertices.FindPointsInBall(P, FMathd::ZeroTolerance, [this, P](FIndex2i IDs)
+									{
+										FVector3d Pos = CellMeshes[IDs.A].AugMesh.GetVertex(IDs.B);
+										return P.DistanceSquared(Pos);
+									}, Update.IDs);
+							}
+						}
+					}
+				}
+			}
+			if (!bAnyUpdatesNeeded)
+			{
+				break;
+			}
+			for (int CellIdx = 0; CellIdx < NumCells; CellIdx++)
+			{
+				CellUnmoved[CellIdx] = true;
+			}
+			
+			// todo: maybe can parallelize if movements are not applied until after?
+			for (int CellIdx = 0; CellIdx < NumCells; CellIdx++)
+			{
+				FDynamicMesh3& Mesh = CellMeshes[CellIdx].AugMesh;
+				TDynamicMeshVertexAttribute<double, 3>* OriginalPosns =
+					static_cast<TDynamicMeshVertexAttribute<double, 3>*>(Mesh.Attributes()->GetAttachedAttribute(OriginalPositionAttribute));
+				auto InterpVert = [&Mesh, &OriginalPosns](int VID, double t)
+				{
+					FVector3d OrigPos, NoisePos;
+					OriginalPosns->GetValue(VID, OrigPos);
+					NoisePos = Mesh.GetVertex(VID);
+					return FVector3d::Lerp(OrigPos, NoisePos, t);
+				};
+				auto InterpTri = [&Mesh, &InterpVert](int TID, double t)
+				{
+					FIndex3i TriVIDs = Mesh.GetTriangle(TID);
+					FTriangle3d Tri;
+					for (int i = 0; i < 3; i++)
+					{
+						Tri.V[i] = InterpVert(TriVIDs[i], t);
+					}
+					return Tri;
+				};
+				auto TestIntersection = [&InterpTri](int TIDA, int TIDB, double t)
+				{
+					FIntrTriangle3Triangle3d TriTri(InterpTri(TIDA, t), InterpTri(TIDB, t));
+					return TriTri.Find();
+				};
+				// resolve tri-tri intersections on this cell's mesh (moving associated verts on other meshes as needed also)
+				for (FUpdate& Update : Updates[CellIdx])
+				{
+					double tsafe = 0;
+					double tbad = 1;
+					if (!TestIntersection(Update.Tris.A, Update.Tris.B, tbad))
+					{
+						continue;
+					}
+					for (int SearchSteps = 0; SearchSteps < 4; SearchSteps++)
+					{
+						double tmid = (tsafe + tbad) * .5;
+						if (TestIntersection(Update.Tris.A, Update.Tris.B, tmid))
+						{
+							tbad = tmid;
+						}
+						else
+						{
+							tsafe = tmid;
+						}
+					}
+					CellUnmoved[CellIdx] = false;
+					for (FIndex2i IDs : Update.IDs)
+					{
+						FVector3d OldPos = CellMeshes[IDs.A].AugMesh.GetVertex(IDs.B);
+						FVector3d NewPos;
+						if (IDs.A == CellIdx)
+						{
+							NewPos = InterpVert(IDs.B, tsafe);
+							Mesh.SetVertex(IDs.B, NewPos);
+						}
+						else
+						{
+							CellUnmoved[IDs.A] = false;
+							FDynamicMesh3& OtherMesh = CellMeshes[IDs.A].AugMesh;
+							TDynamicMeshVertexAttribute<double, 3>* OtherOriginalPosns =
+								static_cast<TDynamicMeshVertexAttribute<double, 3>*>(OtherMesh.Attributes()->GetAttachedAttribute(OriginalPositionAttribute));
+							FVector3d OrigPos;
+							OtherOriginalPosns->GetValue(IDs.B, OrigPos);
+							NewPos = FVector3d::Lerp(OrigPos, OldPos, tsafe);
+							OtherMesh.SetVertex(IDs.B, NewPos);
+						}
+						MeshesVertices.UpdatePoint(IDs, OldPos, NewPos);
+					}
+				}
+			}
+		}
+		
+		// clear "original position" attribute now that we have removed self-intersections
+		for (int CellIdx = 0; CellIdx < NumCells; CellIdx++)
+		{
+			FCellInfo& CellInfo = CellMeshes[CellIdx];
+			FDynamicMesh3& Mesh = CellInfo.AugMesh;
+			Mesh.Attributes()->RemoveAttribute(OriginalPositionAttribute);
+		}
 	}
 	void CreateMeshesForSinglePlane(const FPlanarCells& Cells, const FAxisAlignedBox3d& DomainBounds, bool bNoise, double GlobalUVScale)
 	{
