fix internal material IDs not being set correctly for plane cuts by consistently setting all material IDs for plane cut geometry in one pass at the end

diff --git a/Engine/Plugins/Experimental/PlanarCutPlugin/Source/PlanarCut/Private/PlanarCut.cpp b/Engine/Plugins/Experimental/PlanarCutPlugin/Source/PlanarCut/Private/PlanarCut.cpp
--- a/Engine/Plugins/Experimental/PlanarCutPlugin/Source/PlanarCut/Private/PlanarCut.cpp
+++ b/Engine/Plugins/Experimental/PlanarCutPlugin/Source/PlanarCut/Private/PlanarCut.cpp
@@ -884,8 +884,6 @@ private:
 		//}
 
 		FDynamicMesh3 PlaneMesh(true, true, true, false);
-		PlaneMesh.EnableAttributes();
-		PlaneMesh.Attributes()->EnableMaterialID();
 		FVertexInfo PlaneVertInfo;
 		PlaneVertInfo.bHaveC = true;
 		PlaneVertInfo.bHaveUV = true;
@@ -957,8 +955,6 @@ private:
 				Meshes[Side]->AppendTriangle(VertIDs[1][0], VertIDs[1][1], VertIDs[1][2]),
 				Meshes[Side]->AppendTriangle(VertIDs[1][0], VertIDs[1][2], VertIDs[1][3])
 			};
-			Meshes[Side]->Attributes()->GetMaterialID()->SetNewValue(NewTris[0], MID);
-			Meshes[Side]->Attributes()->GetMaterialID()->SetNewValue(NewTris[1], MID);
 			if (Side == 1)
 			{
 				Meshes[Side]->ReverseTriOrientation(NewTris[0]);
@@ -970,6 +966,12 @@ private:
 
 			// re-enable tangents and visibility attributes, since these are lost when we set the mesh to a copy of the plane mesh
 			UE::PlanarCutInternals::AugmentDynamicMesh::Augment(CellMeshes[Side].AugMesh);
+
+			// Set all material IDs to the one plane's corresponding material ID
+			for (int TID : Meshes[Side]->TriangleIndicesItr())
+			{
+				Meshes[Side]->Attributes()->GetMaterialID()->SetNewValue(TID, MID);
+			}
 		}
 	}
 };
