fix plane cuts in fracture creating duplicate geometry:
reset output array before each multimap multifind, because multifind just appends

#rnx

diff --git a/Engine/Plugins/Experimental/PlanarCutPlugin/Source/PlanarCut/Private/PlanarCut.cpp b/Engine/Plugins/Experimental/PlanarCutPlugin/Source/PlanarCut/Private/PlanarCut.cpp
--- a/Engine/Plugins/Experimental/PlanarCutPlugin/Source/PlanarCut/Private/PlanarCut.cpp
+++ b/Engine/Plugins/Experimental/PlanarCutPlugin/Source/PlanarCut/Private/PlanarCut.cpp
@@ -1328,6 +1328,7 @@ struct FDynamicMeshCollection
 		{
 			int32 GeometryIdx = Collection->TransformToGeometryIndex[MeshData.TransformIndex];
 			int32 InternalMaterialID = bSetDefaultInternalMaterialsFromCollection ? InternalSurfaceMaterials.GetDefaultMaterialIDForGeometry(*Collection, GeometryIdx) : InternalSurfaceMaterials.GlobalMaterialID;
+			ToCutIndices.Reset();
 			ParentTransformToChildren.MultiFind(MeshData.TransformIndex, ToCutIndices);
 
 			// if there's only one mesh here, i.e. it didn't get cut at all
