Fix fracture UVs by recomputing UVs after applying noise, in case when noise can change what the UV should be

#rnx

diff --git a/Engine/Plugins/Experimental/PlanarCutPlugin/Source/PlanarCut/Private/PlanarCut.cpp b/Engine/Plugins/Experimental/PlanarCutPlugin/Source/PlanarCut/Private/PlanarCut.cpp
--- a/Engine/Plugins/Experimental/PlanarCutPlugin/Source/PlanarCut/Private/PlanarCut.cpp
+++ b/Engine/Plugins/Experimental/PlanarCutPlugin/Source/PlanarCut/Private/PlanarCut.cpp
@@ -487,8 +487,8 @@ private:
 				{
 					FVector3d Position = FVector3d(Cells.PlaneBoundaryVertices[BoundaryVertex]);
 					FVector2f UV = FVector2f(PlaneFrame.ToPlaneUV(Position));
-					MinUV.X = FMathf::Min(PlaneVertInfo.UV.X, MinUV.X);
-					MinUV.Y = FMathf::Min(PlaneVertInfo.UV.Y, MinUV.Y);
+					MinUV.X = FMathf::Min(UV.X, MinUV.X);
+					MinUV.Y = FMathf::Min(UV.Y, MinUV.Y);
 				}
 				for (int BoundaryVertex : PlaneBoundary)
 				{
@@ -632,23 +632,14 @@ private:
 				PlaneVertInfo.bHaveUV = true;
 				PlaneVertInfo.bHaveN = true;
 				PlaneVertInfo.Normal = Normal;
+				PlaneVertInfo.UV = FVector2f(0, 0); // UVs will be set below, after noise is added
 				PlaneVertInfo.Color = FVector3f(1, 1, 1);
 
-				FVector2f MinUV(FMathf::MaxReal, FMathf::MaxReal);
-				for (int BoundaryVertex : PlaneBoundary)
-				{
-					FVector3d Position = FVector3d(Cells.PlaneBoundaryVertices[BoundaryVertex]);
-					FVector2f UV = FVector2f(PlaneFrame.ToPlaneUV(Position));
-					MinUV.X = FMathf::Min(PlaneVertInfo.UV.X, MinUV.X);
-					MinUV.Y = FMathf::Min(PlaneVertInfo.UV.Y, MinUV.Y);
-				}
-
 				FPolygon2f Polygon;
 				for (int BoundaryVertex : PlaneBoundary)
 				{
 					PlaneVertInfo.Position = FVector3d(Cells.PlaneBoundaryVertices[BoundaryVertex]);
-					PlaneVertInfo.UV = (FVector2f(PlaneFrame.ToPlaneUV(PlaneVertInfo.Position)) - MinUV) * GlobalUVScale;
-					Polygon.AppendVertex(PlaneVertInfo.UV);
+					Polygon.AppendVertex((FVector2f)PlaneFrame.ToPlaneUV(PlaneVertInfo.Position));
 					Mesh.AppendVertex(PlaneVertInfo);
 				}
 
@@ -863,6 +854,56 @@ private:
 			FDynamicMesh3& Mesh = CellInfo.AugMesh;
 			Mesh.Attributes()->RemoveAttribute(OriginalPositionAttribute);
 		}
+
+		// recompute UVs using new positions after noise was applied + fixed
+		TArray<FVector2f> PlaneMinUVs; PlaneMinUVs.Init(FVector2f(FMathf::MaxReal, FMathf::MaxReal), Cells.Planes.Num());
+		TArray<FFrame3d> PlaneFrames; PlaneFrames.Reserve(Cells.Planes.Num());
+		for (int PlaneIdx = 0; PlaneIdx < Cells.Planes.Num(); PlaneIdx++)
+		{
+			PlaneFrames.Emplace(Cells.Planes[PlaneIdx]);
+		}
+		// first pass to compute min UV for each plane
+		for (FCellInfo& CellInfo : CellMeshes)
+		{
+			FDynamicMesh3& Mesh = CellInfo.AugMesh;
+			FDynamicMeshMaterialAttribute* MaterialIDs = Mesh.Attributes()->GetMaterialID();
+			
+			for (int TID : Mesh.TriangleIndicesItr())
+			{
+				int PlaneIdx = MaterialToPlane(MaterialIDs->GetValue(TID));
+				if (PlaneIdx > -1)
+				{
+					FIndex3i Tri = Mesh.GetTriangle(TID);
+					for (int Idx = 0; Idx < 3; Idx++)
+					{
+						FVector2f UV = FVector2f(PlaneFrames[PlaneIdx].ToPlaneUV(Mesh.GetVertex(Tri[Idx])));
+						FVector2f& MinUV = PlaneMinUVs[PlaneIdx];
+						MinUV.X = FMathf::Min(UV.X, MinUV.X);
+						MinUV.Y = FMathf::Min(UV.Y, MinUV.Y);
+					}
+				}
+			}
+		}
+		// second pass to actually set UVs
+		for (FCellInfo& CellInfo : CellMeshes)
+		{
+			FDynamicMesh3& Mesh = CellInfo.AugMesh;
+			FDynamicMeshMaterialAttribute* MaterialIDs = Mesh.Attributes()->GetMaterialID();
+
+			for (int TID : Mesh.TriangleIndicesItr())
+			{
+				int PlaneIdx = MaterialToPlane(MaterialIDs->GetValue(TID));
+				if (PlaneIdx > -1)
+				{
+					FIndex3i Tri = Mesh.GetTriangle(TID);
+					for (int Idx = 0; Idx < 3; Idx++)
+					{
+						FVector2f UV = ((FVector2f)PlaneFrames[PlaneIdx].ToPlaneUV(Mesh.GetVertex(Tri[Idx])) - PlaneMinUVs[PlaneIdx]) * GlobalUVScale;
+						Mesh.SetVertexUV(Tri[Idx], UV);
+					}
+				}
+			}
+		}
 	}
 	void CreateMeshesForSinglePlane(const FPlanarCells& Cells, const FAxisAlignedBox3d& DomainBounds, bool bNoise, double GlobalUVScale)
 	{
