Add grout option for all fracture types

diff --git a/Engine/Plugins/Experimental/ChaosEditor/Source/FractureEditor/Private/FractureToolBrick.cpp b/Engine/Plugins/Experimental/ChaosEditor/Source/FractureEditor/Private/FractureToolBrick.cpp
--- a/Engine/Plugins/Experimental/ChaosEditor/Source/FractureEditor/Private/FractureToolBrick.cpp
+++ b/Engine/Plugins/Experimental/ChaosEditor/Source/FractureEditor/Private/FractureToolBrick.cpp
@@ -293,12 +293,18 @@ int32 UFractureToolBrick::ExecuteFracture(const FFractureToolContext& FractureCo
 		const float HalfDepth = BrickSettings->BrickDepth * 0.5f;
 		const float HalfLength = BrickSettings->BrickLength * 0.5f;
 
-		FVector Max(HalfLength, HalfDepth, HalfHeight);
 		TArray<FBox> BricksToCut;
 
+		// space the bricks by the grout setting, constrained to not erase the bricks or have zero grout
+		// (currently zero grout bricks would break assumptions in the fracture)
+		float MinDim = FMath::Min3(HalfHeight, HalfDepth, HalfLength);
+		float HalfGrout = FMath::Clamp(0.5f * CutterSettings->Grout, MinDim * 0.02f, MinDim * 0.98f);
+		FVector HalfBrick(HalfLength - HalfGrout, HalfDepth - HalfGrout, HalfHeight - HalfGrout);
+		FBox BrickBox(-HalfBrick, HalfBrick);
+
 		for (const FTransform& Trans : BrickTransforms)
 		{
-			BricksToCut.Add(FBox(-Max * 0.95f, Max * 0.95f).TransformBy(Trans));
+			BricksToCut.Add(BrickBox.TransformBy(Trans));
 		}
 
  		FPlanarCells VoronoiPlanarCells = FPlanarCells(BricksToCut);
@@ -313,7 +319,7 @@ int32 UFractureToolBrick::ExecuteFracture(const FFractureToolContext& FractureCo
 			VoronoiPlanarCells.InternalSurfaceMaterials.NoiseSettings = NoiseSettings;
 		}
 
-		return CutMultipleWithPlanarCells(VoronoiPlanarCells, *FractureContext.GetGeometryCollection(), FractureContext.GetSelection());
+		return CutMultipleWithPlanarCells(VoronoiPlanarCells, *FractureContext.GetGeometryCollection(), FractureContext.GetSelection(), 0);
 	}
 
 	return INDEX_NONE;
diff --git a/Engine/Plugins/Experimental/ChaosEditor/Source/FractureEditor/Private/FractureToolCutter.cpp b/Engine/Plugins/Experimental/ChaosEditor/Source/FractureEditor/Private/FractureToolCutter.cpp
--- a/Engine/Plugins/Experimental/ChaosEditor/Source/FractureEditor/Private/FractureToolCutter.cpp
+++ b/Engine/Plugins/Experimental/ChaosEditor/Source/FractureEditor/Private/FractureToolCutter.cpp
@@ -159,11 +159,10 @@ int32 UFractureToolVoronoiCutterBase::ExecuteFracture(const FFractureToolContext
 		TArray<FVector> Sites;
 		GenerateVoronoiSites(FractureContext, Sites);
 		FBox VoronoiBounds = FractureContext.GetBounds();
-		if (CutterSettings->Amplitude > 0.0f)
-		{
-			// expand bounds to make sure noise-perturbed voronoi cells still contain the whole input mesh
-			VoronoiBounds = VoronoiBounds.ExpandBy(CutterSettings->Amplitude);
-		}
+			
+		// expand bounds to make sure noise-perturbed voronoi cells still contain the whole input mesh
+		VoronoiBounds = VoronoiBounds.ExpandBy(CutterSettings->Amplitude + CutterSettings->Grout);
+			
 		FVoronoiDiagram Voronoi(Sites, VoronoiBounds, .1f);
 
 		FPlanarCells VoronoiPlanarCells = FPlanarCells(Sites, Voronoi);
@@ -178,7 +177,7 @@ int32 UFractureToolVoronoiCutterBase::ExecuteFracture(const FFractureToolContext
 			VoronoiPlanarCells.InternalSurfaceMaterials.NoiseSettings = NoiseSettings;
 		}
 
-		return CutMultipleWithPlanarCells(VoronoiPlanarCells, *(FractureContext.GetGeometryCollection()), FractureContext.GetSelection());
+		return CutMultipleWithPlanarCells(VoronoiPlanarCells, *(FractureContext.GetGeometryCollection()), FractureContext.GetSelection(), CutterSettings->Grout);
 	}
 
 	return INDEX_NONE;
diff --git a/Engine/Plugins/Experimental/ChaosEditor/Source/FractureEditor/Private/FractureToolCutter.h b/Engine/Plugins/Experimental/ChaosEditor/Source/FractureEditor/Private/FractureToolCutter.h
--- a/Engine/Plugins/Experimental/ChaosEditor/Source/FractureEditor/Private/FractureToolCutter.h
+++ b/Engine/Plugins/Experimental/ChaosEditor/Source/FractureEditor/Private/FractureToolCutter.h
@@ -49,6 +49,10 @@ public:
 	UPROPERTY(EditAnywhere, Category = CommonFracture, meta = (DisplayName = "Draw Diagram"))
 	bool bDrawDiagram;
 
+	/** Amount of space to leave between cut pieces */
+	UPROPERTY(EditAnywhere, Category = CommonFracture, meta = (UIMin = "0.0", ClampMin = "0.0"))
+	float Grout = 0.0f;
+
 	/** Size of the noise displacement in centimeters */
 	UPROPERTY(EditAnywhere, Category = Noise, meta = (UIMin = "0.0"))
 	float Amplitude;
diff --git a/Engine/Plugins/Experimental/ChaosEditor/Source/FractureEditor/Private/FractureToolPlaneCut.cpp b/Engine/Plugins/Experimental/ChaosEditor/Source/FractureEditor/Private/FractureToolPlaneCut.cpp
--- a/Engine/Plugins/Experimental/ChaosEditor/Source/FractureEditor/Private/FractureToolPlaneCut.cpp
+++ b/Engine/Plugins/Experimental/ChaosEditor/Source/FractureEditor/Private/FractureToolPlaneCut.cpp
@@ -130,7 +130,7 @@ int32 UFractureToolPlaneCut::ExecuteFracture(const FFractureToolContext& Fractur
 			InternalSurfaceMaterials.NoiseSettings = NoiseSettings;
 		}
 
-		return CutMultipleWithMultiplePlanes(CuttingPlanes, InternalSurfaceMaterials, *FractureContext.GetGeometryCollection(), FractureContext.GetSelection());
+		return CutMultipleWithMultiplePlanes(CuttingPlanes, InternalSurfaceMaterials, *FractureContext.GetGeometryCollection(), FractureContext.GetSelection(), CutterSettings->Grout);
 	}
 
 	return INDEX_NONE;
diff --git a/Engine/Plugins/Experimental/ChaosEditor/Source/FractureEditor/Private/FractureToolSlice.cpp b/Engine/Plugins/Experimental/ChaosEditor/Source/FractureEditor/Private/FractureToolSlice.cpp
--- a/Engine/Plugins/Experimental/ChaosEditor/Source/FractureEditor/Private/FractureToolSlice.cpp
+++ b/Engine/Plugins/Experimental/ChaosEditor/Source/FractureEditor/Private/FractureToolSlice.cpp
@@ -169,7 +169,7 @@ int32 UFractureToolSlice::ExecuteFracture(const FFractureToolContext& FractureCo
 			InternalSurfaceMaterials.NoiseSettings = NoiseSettings;
 		}
 
-		return CutMultipleWithMultiplePlanes(CuttingPlanes, InternalSurfaceMaterials, *FractureContext.GetGeometryCollection(), FractureContext.GetSelection());
+		return CutMultipleWithMultiplePlanes(CuttingPlanes, InternalSurfaceMaterials, *FractureContext.GetGeometryCollection(), FractureContext.GetSelection(), CutterSettings->Grout);
 	}
 
 	return INDEX_NONE;
diff --git a/Engine/Plugins/Experimental/PlanarCutPlugin/Source/PlanarCut/Private/PlanarCut.cpp b/Engine/Plugins/Experimental/PlanarCutPlugin/Source/PlanarCut/Private/PlanarCut.cpp
--- a/Engine/Plugins/Experimental/PlanarCutPlugin/Source/PlanarCut/Private/PlanarCut.cpp
+++ b/Engine/Plugins/Experimental/PlanarCutPlugin/Source/PlanarCut/Private/PlanarCut.cpp
@@ -5,6 +5,7 @@
 #include "Async/ParallelFor.h"
 #include "Spatial/FastWinding.h"
 #include "Spatial/PointHashGrid3.h"
+#include "Spatial/MeshSpatialSort.h"
 #include "Arrangement2d.h"
 #include "MeshAdapter.h"
 #include "FrameTypes.h"
@@ -18,6 +19,7 @@
 #include "DynamicMesh3.h"
 #include "DynamicMeshEditor.h"
 #include "DynamicMeshAABBTree3.h"
+#include "MeshTransforms.h"
 #include "Operations/MeshBoolean.h"
 #include "Operations/MeshSelfUnion.h"
 #include "Operations/MergeCoincidentMeshEdges.h"
@@ -261,15 +263,52 @@ struct FCellMeshes
 	FVector NoiseOffsetY;
 	FVector NoiseOffsetZ;
 
-	FCellMeshes() {}
-	FCellMeshes(const FPlanarCells& Cells, FAxisAlignedBox3d DomainBounds, double ExtendDomain, bool bIncludeOutsideCell = true)
+	FCellMeshes()
 	{
-		Init(Cells, DomainBounds, ExtendDomain, bIncludeOutsideCell);
+		InitEmpty();
+	}
+	
+	FCellMeshes(const FPlanarCells& Cells, FAxisAlignedBox3d DomainBounds, double Grout, double ExtendDomain, bool bIncludeOutsideCell)
+	{
+		Init(Cells, DomainBounds, Grout, ExtendDomain, bIncludeOutsideCell);
+	}
+	
+	// Special function to just make the "grout" part of the planar mesh cells
+	// Used to make the multi-plane cuts with grout easier to implement
+	void MakeOnlyPlanarGroutCell(const FPlanarCells& Cells, FAxisAlignedBox3d DomainBounds, double Grout)
+	{
+		CellMeshes.Reset();
+
+		if (!ensure(Grout > 0) || !ensure(Cells.IsInfinitePlane()))
+		{
+			return;
+		}
+		
+		float GlobalUVScale = Cells.InternalSurfaceMaterials.GlobalUVScale;
+		if (!ensure(GlobalUVScale > 0))
+		{
+			GlobalUVScale = 1;
+		}
+
+		CellMeshes.SetNum(1);
+
+		bool bNoise = Cells.InternalSurfaceMaterials.NoiseSettings.IsSet();
+
+		double ExtendDomain = bNoise ? Cells.InternalSurfaceMaterials.NoiseSettings->Amplitude : 0;
+		DomainBounds.Expand(ExtendDomain);
+		
+		CreateMeshesForSinglePlane(Cells, DomainBounds, bNoise, GlobalUVScale, Grout, true);
+
+		for (FCellInfo& CellInfo : CellMeshes)
+		{
+			UE::PlanarCutInternals::AugmentDynamicMesh::SetDefaultAttributes(CellInfo.AugMesh, Cells.InternalSurfaceMaterials);
+		}
 	}
 
 	void RemeshForNoise(FDynamicMesh3& Mesh, EEdgeRefineFlags EdgeFlags, double TargetEdgeLen)
 	{
 		FQueueRemesher Remesh(&Mesh);
+		Remesh.bPreventNormalFlips = true;
 		FMeshConstraints Constraints;
 
 		FMeshBoundaryLoops Boundary(&Mesh);
@@ -362,11 +401,17 @@ struct FCellMeshes
 		return MaterialID >= 0 ? -1 : -(MaterialID+1);
 	}
 
-	void Init(const FPlanarCells& Cells, FAxisAlignedBox3d DomainBounds, double ExtendDomain, bool bIncludeOutsideCell = true)
+	void InitEmpty()
 	{
 		NoiseOffsetX = FMath::VRand() * 100;
 		NoiseOffsetY = FMath::VRand() * 100;
 		NoiseOffsetZ = FMath::VRand() * 100;
+		OutsideCellIndex = -1;
+	}
+
+	void Init(const FPlanarCells& Cells, FAxisAlignedBox3d DomainBounds, double Grout, double ExtendDomain, bool bIncludeOutsideCell)
+	{
+		InitEmpty();
 
 		float GlobalUVScale = Cells.InternalSurfaceMaterials.GlobalUVScale;
 		if (!ensure(GlobalUVScale > 0))
@@ -375,20 +420,8 @@ struct FCellMeshes
 		}
 
 		int NumCells = Cells.NumCells;
-		bool bHasOutsideCell = false;
-		OutsideCellIndex = -1;
+		bool bHasGroutCell = Grout > 0;
 		if (bIncludeOutsideCell && !Cells.IsInfinitePlane())
-		{
-			for (const TPair<int32, int32>& CellPair : Cells.PlaneCells)
-			{
-				if (CellPair.Value == -1)
-				{
-					bHasOutsideCell = true;
-					break;
-				}
-			}
-		}
-		if (bHasOutsideCell)
 		{
 			OutsideCellIndex = NumCells;
 			NumCells++;
@@ -397,26 +430,29 @@ struct FCellMeshes
 		CellMeshes.Reset();
 		CellMeshes.SetNum(NumCells);
 
-		DomainBounds.Expand(ExtendDomain);
-
 		bool bNoise = Cells.InternalSurfaceMaterials.NoiseSettings.IsSet();
 		if (bNoise)
 		{
 			ExtendDomain += Cells.InternalSurfaceMaterials.NoiseSettings->Amplitude;
 		}
+		DomainBounds.Expand(ExtendDomain);
 
 		// special handling for the infinite plane case; we need to adapt this to be a closed volume
 		if (Cells.IsInfinitePlane())
 		{
-			CreateMeshesForSinglePlane(Cells, DomainBounds, bNoise, GlobalUVScale);
+			CreateMeshesForSinglePlane(Cells, DomainBounds, bNoise, GlobalUVScale, Grout, false);
 		}
-		else if (!bNoise) // bounded cells w/ no noise
-		{
-			CreateMeshesForBoundedPlanesWithoutNoise(NumCells, Cells, DomainBounds, bNoise, GlobalUVScale);
-		}
-		else // bounded cells with noise -- make each boundary plane separately so we can remesh them w/ noise vertices
+		else
 		{
-			CreateMeshesForBoundedPlanesWithNoise(NumCells, Cells, DomainBounds, bNoise, GlobalUVScale);
+			if (!bNoise) // bounded cells w/ no noise
+			{
+				CreateMeshesForBoundedPlanesWithoutNoise(NumCells, Cells, DomainBounds, bNoise, GlobalUVScale);
+			}
+			else // bounded cells with noise -- make each boundary plane separately so we can remesh them w/ noise vertices
+			{
+				CreateMeshesForBoundedPlanesWithNoise(NumCells, Cells, DomainBounds, bNoise, GlobalUVScale);
+			}
+			ApplyGeneralGrout(Grout);
 		}
 		
 		// TODO: self-union on cells when it makes sense to do so (for non-single-plane inputs w/ high noise or possible untracked adjacencies)
@@ -433,6 +469,67 @@ struct FCellMeshes
 		}
 	}
 
+	void ApplyGeneralGrout(double Grout)
+	{
+		if (Grout <= 0)
+		{
+			return;
+		}
+
+		// apply grout to all cells
+		for (int MeshIdx = 0; MeshIdx < CellMeshes.Num(); MeshIdx++)
+		{
+			if (MeshIdx == OutsideCellIndex)
+			{
+				continue;
+			}
+
+			FDynamicMesh3& Mesh = CellMeshes[MeshIdx].AugMesh;
+			// TODO: scale from mesh center of mass instead of the vertex centroid?
+			FVector3d VertexCentroid(0, 0, 0);
+			for (FVector3d V : Mesh.VerticesItr())
+			{
+				VertexCentroid += V;
+			}
+			VertexCentroid /= Mesh.VertexCount();
+			FAxisAlignedBox3d Bounds = Mesh.GetCachedBounds();
+			double BoundsSize = Bounds.MaxDim();
+			// currently just scale the meshes down so they leave half-a-grout worth of space on their longest axis
+			// or delete the mesh if it's so small that that would require a negative scale
+			// TODO: consider instead computing a true offset mesh
+			//  (note that we don't currently have a good UV-preserving+sharp-edge-preserving way to do that)
+			double ScaleFactor = (BoundsSize - Grout*.5) / BoundsSize;
+			if (ScaleFactor < FMathd::ZeroTolerance * 1000)
+			{
+				// if the grout scale factor would be ~zero or negative, just clear the mesh instead
+				Mesh.Clear();
+				UE::PlanarCutInternals::AugmentDynamicMesh::Augment(Mesh);
+			}
+			else
+			{
+				MeshTransforms::Scale(Mesh, FVector3d::One() * ScaleFactor, VertexCentroid);
+			}
+		}
+
+		// create outside cell (if there is room for it) by appending all the other meshes
+		if (OutsideCellIndex != -1)
+		{
+			FDynamicMesh3& OutsideMesh = CellMeshes[OutsideCellIndex].AugMesh;
+			OutsideMesh.Clear();
+			UE::PlanarCutInternals::AugmentDynamicMesh::Augment(OutsideMesh);
+			FDynamicMeshEditor OutsideMeshEditor(&OutsideMesh);
+			for (int MeshIdx = 0; MeshIdx < CellMeshes.Num(); MeshIdx++)
+			{
+				if (MeshIdx == OutsideCellIndex)
+				{
+					continue;
+				}
+				FMeshIndexMappings IndexMaps;
+				OutsideMeshEditor.AppendMesh(&CellMeshes[MeshIdx].AugMesh, IndexMaps);
+			}
+		}
+	}
+
 	void AppendMesh(FDynamicMesh3& Base, FDynamicMesh3& ToAppend, bool bFlipped)
 	{
 		FDynamicMeshEditor Editor(&Base);
@@ -905,8 +1002,11 @@ private:
 			}
 		}
 	}
-	void CreateMeshesForSinglePlane(const FPlanarCells& Cells, const FAxisAlignedBox3d& DomainBounds, bool bNoise, double GlobalUVScale)
+	
+	void CreateMeshesForSinglePlane(const FPlanarCells& Cells, const FAxisAlignedBox3d& DomainBounds, bool bNoise, double GlobalUVScale, double Grout, bool bOnlyGrout)
 	{
+		bool bHasGrout = Grout > 0;
+
 		int MID = PlaneToMaterial(0);
 		FPlane Plane = Cells.Planes[0];
 
@@ -948,70 +1048,116 @@ private:
 			ApplyNoise(PlaneMesh, PlaneFrame.GetAxis(2), Cells.InternalSurfaceMaterials.NoiseSettings.GetValue(), true);
 			FMeshNormals::QuickComputeVertexNormals(PlaneMesh);
 		}
-		FDynamicMesh3* Meshes[2]{ &CellMeshes[0].AugMesh, &CellMeshes[1].AugMesh };
-		for (int Side = 0; Side < 2; Side++)
+		TArray<int> PlaneBoundary,  // loop of vertex IDs on the boundary of PlaneMesh (starting with vertex 0)
+			PlaneBoundaryCornerIndices; // indices of the corner vertices in the PlaneBoundary array
 		{
-			*Meshes[Side] = PlaneMesh;
 			double Offset = ZRange.Max;
-			FMeshBoundaryLoops Boundary(Meshes[Side]);
+			FMeshBoundaryLoops Boundary(&PlaneMesh);
 			checkSlow(Boundary.GetLoopCount() == 1);
 			int FirstIdx;
 			bool bFound = Boundary[0].Vertices.Find(0, FirstIdx);
 			checkSlow(bFound);
-			TArray<int> VertIDs[2], MatchedIndices[2];
-			VertIDs[0] = Boundary[0].Vertices;
+			PlaneBoundary = Boundary[0].Vertices;
 			if (FirstIdx != 0)
 			{
-				Algo::Rotate(VertIDs[0], FirstIdx);
+				Algo::Rotate(PlaneBoundary, FirstIdx);
 			}
-			checkSlow(VertIDs[0][0] == 0);
-			MatchedIndices[0].Add(0);
+			checkSlow(PlaneBoundary[0] == 0);
+
+			PlaneBoundaryCornerIndices.Add(0);
 			int FoundIndices = 1;
-			for (int VIDIdx = 0; VIDIdx < VertIDs[0].Num(); VIDIdx++)
+			for (int VIDIdx = 0; VIDIdx < PlaneBoundary.Num(); VIDIdx++)
 			{
-				int VID = VertIDs[0][VIDIdx];
+				int VID = PlaneBoundary[VIDIdx];
 				if (VID == FoundIndices)
 				{
 					FoundIndices++;
-					MatchedIndices[0].Add(VIDIdx);
+					PlaneBoundaryCornerIndices.Add(VIDIdx);
 				}
 			}
-
-			if (Side == 0)
+		}
+		FDynamicMesh3* Meshes[2];
+		if (!bOnlyGrout)
+		{
+			for (int Side = 0; Side < 2; Side++)
 			{
-				Meshes[Side]->ReverseOrientation(true);
-				Offset = ZRange.Min;
-			}
-			PlaneVertInfo.Normal = FVector3f(Plane.GetNormal()) * (-1 + Side * 2);
-			FVector3d OffsetVec = FVector3d(Plane.GetNormal()) * Offset;
+				Meshes[Side] = &CellMeshes[Side].AugMesh;
+				*Meshes[Side] = PlaneMesh;
+				double Offset = ZRange.Max;
+				TArray<int> CapBoundary, CapBoundaryCornerIndices;
 
-			for (int CornerIdx = 0; CornerIdx < 4; CornerIdx++)
+				if (Side == 0)
+				{
+					Meshes[Side]->ReverseOrientation(true);
+					Offset = ZRange.Min;
+				}
+				PlaneVertInfo.Normal = FVector3f(Plane.GetNormal()) * (-1 + Side * 2);
+				FVector3d OffsetVec = FVector3d(Plane.GetNormal()) * Offset;
+
+				for (int CornerIdx = 0; CornerIdx < 4; CornerIdx++)
+				{
+					PlaneVertInfo.Position = Meshes[Side]->GetVertex(CornerIdx) + OffsetVec;
+					// UVs shouldn't matter for outer box vertices because they're outside of the domain by construction ...
+					CapBoundary.Add(Meshes[Side]->AppendVertex(PlaneVertInfo));
+					CapBoundaryCornerIndices.Add(CornerIdx);
+				}
+				int NewTris[2]{
+					Meshes[Side]->AppendTriangle(CapBoundary[0], CapBoundary[1], CapBoundary[2]),
+					Meshes[Side]->AppendTriangle(CapBoundary[0], CapBoundary[2], CapBoundary[3])
+				};
+				if (Side == 1)
+				{
+					Meshes[Side]->ReverseTriOrientation(NewTris[0]);
+					Meshes[Side]->ReverseTriOrientation(NewTris[1]);
+				}
+				FDynamicMeshEditor Editor(Meshes[Side]);
+				FDynamicMeshEditResult ResultOut;
+				Editor.StitchSparselyCorrespondedVertexLoops(PlaneBoundary, PlaneBoundaryCornerIndices, CapBoundary, CapBoundaryCornerIndices, ResultOut, Side == 0);
+			}
+		}
+		if (bHasGrout)
+		{
+			int GroutIdx = bOnlyGrout ? 0 : 2;
+			FDynamicMesh3* GroutMesh = &CellMeshes[GroutIdx].AugMesh;
+			FVector3d GroutOffset = Plane.GetNormal() * (Grout * .5);
+			if (!bOnlyGrout)
 			{
-				PlaneVertInfo.Position = Meshes[Side]->GetVertex(CornerIdx) + OffsetVec;
-				// UVs shouldn't matter for outer box vertices because they're outside of the domain by construction ...
-				VertIDs[1].Add(Meshes[Side]->AppendVertex(PlaneVertInfo));
-				MatchedIndices[1].Add(CornerIdx);
+				for (int Side = 0; Side < 2; Side++)
+				{
+					// shift both sides out by Grout/2
+					MeshTransforms::Translate(*Meshes[Side], GroutOffset * (-1 + Side * 2));
+				}
 			}
-			int NewTris[2]{
-				Meshes[Side]->AppendTriangle(VertIDs[1][0], VertIDs[1][1], VertIDs[1][2]),
-				Meshes[Side]->AppendTriangle(VertIDs[1][0], VertIDs[1][2], VertIDs[1][3])
-			};
-			if (Side == 1)
+
+			// make the center (grout) by stitching together two offset copies of PlaneMesh
+			*GroutMesh = PlaneMesh;
+			GroutMesh->ReverseOrientation(true);
+			MeshTransforms::Translate(*GroutMesh, GroutOffset);
+			FMeshIndexMappings IndexMaps;
+			FDynamicMeshEditor Editor(GroutMesh);
+			Editor.AppendMesh(&PlaneMesh, IndexMaps, [GroutOffset](int VID, const FVector3d& PosIn) {return PosIn - GroutOffset;});
+			TArray<int> AppendPlaneBoundary; AppendPlaneBoundary.Reserve(PlaneBoundary.Num());
+			TArray<int> RevBoundary = PlaneBoundary;
+			Algo::Reverse(RevBoundary);
+			for (int VID : RevBoundary)
 			{
-				Meshes[Side]->ReverseTriOrientation(NewTris[0]);
-				Meshes[Side]->ReverseTriOrientation(NewTris[1]);
+				AppendPlaneBoundary.Add(IndexMaps.GetNewVertex(VID));
 			}
-			FDynamicMeshEditor Editor(Meshes[Side]);
 			FDynamicMeshEditResult ResultOut;
-			Editor.StitchSparselyCorrespondedVertexLoops(VertIDs[0], MatchedIndices[0], VertIDs[1], MatchedIndices[1], ResultOut, Side == 0);
+			Editor.StitchVertexLoopsMinimal(RevBoundary, AppendPlaneBoundary, ResultOut);
+		}
 
+		// fix up custom attributes and material IDs for all meshes
+		for (int CellIdx = 0; CellIdx < CellMeshes.Num(); CellIdx++)
+		{
+			FDynamicMesh3& Mesh = CellMeshes[CellIdx].AugMesh;
 			// re-enable tangents and visibility attributes, since these are lost when we set the mesh to a copy of the plane mesh
-			UE::PlanarCutInternals::AugmentDynamicMesh::Augment(CellMeshes[Side].AugMesh);
+			UE::PlanarCutInternals::AugmentDynamicMesh::Augment(Mesh);
 
 			// Set all material IDs to the one plane's corresponding material ID
-			for (int TID : Meshes[Side]->TriangleIndicesItr())
+			for (int TID : Mesh.TriangleIndicesItr())
 			{
-				Meshes[Side]->Attributes()->GetMaterialID()->SetNewValue(TID, MID);
+				Mesh.Attributes()->GetMaterialID()->SetNewValue(TID, MID);
 			}
 		}
 	}
@@ -1112,6 +1258,7 @@ struct FDynamicMeshCollection
 
 	int32 CutWithMultiplePlanes(
 		const TArrayView<const FPlane>& Planes, 
+		double Grout,
 		FGeometryCollection* Collection,
 		FInternalSurfaceMaterials& InternalSurfaceMaterials,
 		bool bSetDefaultInternalMaterialsFromCollection
@@ -1133,6 +1280,48 @@ struct FDynamicMeshCollection
 		auto EnterProgressFrame = [](float Progress) {};
 #endif
 
+		bool bHasGrout = Grout > 0;
+
+		if (bHasGrout)
+		{
+			// For multi-plane cuts with grout specifically, the easiest path seems to be:
+			// 1. Build the "grout" section of each plane
+			// 2. Take the union of all those grout sections as the grout mesh
+			// 3. Use the generic CutWithCellMeshes path, where that grout mesh is both the inner and outside cell mesh
+			//    (Note the outside cell mesh is subtracted, not intersected)
+			//    (Note this relies on island splitting to separate all the pieces afterwards.)
+			FCellMeshes GroutCells;
+			GroutCells.CellMeshes.SetNum(2);
+			FDynamicMesh3& GroutMesh = GroutCells.CellMeshes[0].AugMesh;
+			FDynamicMeshEditor GroutAppender(&GroutMesh);
+			FMeshIndexMappings IndexMaps;
+			for (int32 PlaneIdx = 0; PlaneIdx < Planes.Num(); PlaneIdx++)
+			{
+				EnterProgressFrame(.5);
+				FPlanarCells PlaneCells(Planes[PlaneIdx]);
+				PlaneCells.InternalSurfaceMaterials = InternalSurfaceMaterials;
+				FCellMeshes PlaneGroutMesh;
+				PlaneGroutMesh.MakeOnlyPlanarGroutCell(PlaneCells, Bounds, Grout);
+				GroutAppender.AppendMesh(&PlaneGroutMesh.CellMeshes[0].AugMesh, IndexMaps);
+			}
+
+			EnterProgressFrame(Planes.Num() * .2);
+			FMeshSelfUnion GroutUnion(&GroutMesh);
+			GroutUnion.bWeldSharedEdges = false;
+			GroutUnion.Compute();
+
+			EnterProgressFrame(Planes.Num() * .1);
+			// first mesh is the same as the second mesh, but will be subtracted b/c it's the "outside cell"
+			GroutCells.CellMeshes[1].AugMesh = GroutMesh;
+			GroutCells.OutsideCellIndex = 1;
+
+			EnterProgressFrame(Planes.Num() * .2);
+			TArray<TPair<int32, int32>> CellConnectivity;
+			CellConnectivity.Add(TPair<int32, int32>(0, -1));
+
+			return CutWithCellMeshes(InternalSurfaceMaterials, CellConnectivity, GroutCells, Collection, bSetDefaultInternalMaterialsFromCollection);
+		}
+
 		bool bHasProximity = Collection->HasAttribute("Proximity", FGeometryCollection::GeometryGroup);
 		TArray<TUniquePtr<FMeshData>> ToCut;
 		TArray<TUniquePtr<TPointHashGrid3d<int>>> VerticesHashes;
@@ -1215,7 +1404,7 @@ struct FDynamicMeshCollection
 			EnterProgressFrame(1);
 			FPlanarCells PlaneCells(Planes[PlaneIdx]);
 			PlaneCells.InternalSurfaceMaterials = InternalSurfaceMaterials;
-			FCellMeshes CellMeshes(PlaneCells, Bounds, InternalSurfaceMaterials.NoiseSettings.IsSet() ? InternalSurfaceMaterials.NoiseSettings->Amplitude : 0, false);
+			FCellMeshes CellMeshes(PlaneCells, Bounds, 0, 0, false);
 
 			// TODO: we could do these cuts in parallel (will takes some rework of the proximity and how results are added to the ToCut array)
 			for (int32 ToCutIdx = 0, ToCutNum = ToCut.Num(); ToCutIdx < ToCutNum; ToCutIdx++)
@@ -1413,13 +1602,14 @@ struct FDynamicMeshCollection
 	/**
 	 * Cut collection meshes with cell meshes, and append results to a geometry collection
 	 *
-	 * @param PlanarCells The source definitions of the cells to cut with
+	 * @param InternalSurfaceMaterials Internal material info (used for material ID)
+	 * @param CellConnectivity The connectivity between cells: PlaneTag -> The two cells separated by triangles with this tag
 	 * @param CellsMeshes Meshed versions of the cells, with noise and material properties baked in
 	 * @param Collection Results will be stored in this
 	 * @param bSetDefaultInternalMaterialsFromCollection If true, set internal materials to the most common external material + 1, following a convenient artist convention
 	 * @return Index of the first created geometry
 	 */
-	int32 CutWithCellMeshes(const FPlanarCells& PlanarCells, FCellMeshes& CellMeshes, FGeometryCollection* Collection, bool bSetDefaultInternalMaterialsFromCollection)
+	int32 CutWithCellMeshes(const FInternalSurfaceMaterials& InternalSurfaceMaterials, const TArray<TPair<int32, int32>>& CellConnectivity, FCellMeshes& CellMeshes, FGeometryCollection* Collection, bool bSetDefaultInternalMaterialsFromCollection)
 	{
 		// TODO: should we do these cuts in parallel, and the appends sequentially below?
 		int32 FirstIdx = -1;
@@ -1468,7 +1658,7 @@ struct FDynamicMeshCollection
 				TMultiMap<int32, int32> CellToGeometry;
 				TMap<int32, int32> GeometryToResultMesh;
 				int32 SubPartIndex = 0;
-				int32 InternalMaterialID = bSetDefaultInternalMaterialsFromCollection ? PlanarCells.InternalSurfaceMaterials.GetDefaultMaterialIDForGeometry(*Collection, GeometryIdx) : PlanarCells.InternalSurfaceMaterials.GlobalMaterialID;
+				int32 InternalMaterialID = bSetDefaultInternalMaterialsFromCollection ? InternalSurfaceMaterials.GetDefaultMaterialIDForGeometry(*Collection, GeometryIdx) : InternalSurfaceMaterials.GlobalMaterialID;
 
 				for (int32 CellIdx = 0; CellIdx < CellMeshes.CellMeshes.Num(); CellIdx++)
 				{					
@@ -1535,7 +1725,7 @@ struct FDynamicMeshCollection
 					};
 					for (int32 PlaneIdx : PlanesInOutput)
 					{
-						TPair<int32, int32> Cells = PlanarCells.PlaneCells[PlaneIdx];
+						TPair<int32, int32> Cells = CellConnectivity[PlaneIdx];
 						int32 SecondCell = Cells.Value < 0 ? CellMeshes.OutsideCellIndex : Cells.Value;
 						if (SecondCell != -1)
 						{
@@ -1658,10 +1848,41 @@ struct FDynamicMeshCollection
 			VertComponents.Union(Tri.C, Tri.A);
 		}
 		
-		return FDynamicMeshEditor::SplitMesh(&Source, SeparatedMeshes, [&Source, &VertComponents](int TID)
+		bool bWasSplit = FDynamicMeshEditor::SplitMesh(&Source, SeparatedMeshes, [&Source, &VertComponents](int TID)
 			{
 				return (int)VertComponents.Find(Source.GetTriangle(TID).A);
 			});
+
+		if (bWasSplit)
+		{
+			// disconnected components that are contained inside other components need to be re-merged
+			TMeshSpatialSort<FDynamicMesh3> SpatialSort(SeparatedMeshes);
+			SpatialSort.NestingMethod = TMeshSpatialSort<FDynamicMesh3>::ENestingMethod::InLargestParent;
+			SpatialSort.bOnlyNestNegativeVolumes = false;
+			SpatialSort.bOnlyParentPostiveVolumes = true;
+			SpatialSort.Compute();
+			TArray<bool> KeepMeshes; KeepMeshes.Init(true, SeparatedMeshes.Num());
+			for (TMeshSpatialSort<FDynamicMesh3>::FMeshNesting& Nest : SpatialSort.Nests)
+			{
+				FDynamicMeshEditor Editor(&SeparatedMeshes[Nest.OuterIndex]);
+				FMeshIndexMappings Mappings;
+				for (int Inner : Nest.InnerIndices)
+				{
+					Editor.AppendMesh(&SeparatedMeshes[Inner], Mappings);
+					KeepMeshes[Inner] = false;
+				}
+			}
+			for (int Idx = 0; Idx < SeparatedMeshes.Num(); Idx++)
+			{
+				if (!KeepMeshes[Idx])
+				{
+					SeparatedMeshes.RemoveAtSwap(Idx, 1, false);
+					KeepMeshes.RemoveAtSwap(Idx, 1, false);
+					Idx--;
+				}
+			}
+		}
+		return bWasSplit;
 	}
 
 	static int32 AppendToCollection(const FTransform& ToCollection, FDynamicMesh3& Mesh, int32 TransformParent, int32 SubPartIndex, FGeometryCollection& Output, int32 InternalMaterialID)
@@ -3804,6 +4025,7 @@ int32 CutWithPlanarCells(
 	FPlanarCells& Cells,
 	FGeometryCollection& Source,
 	int32 TransformIdx,
+	double Grout,
 	const TOptional<FTransform>& TransformCells,
 	bool bIncludeOutsideCellInOutput,
 	float CheckDistanceAcrossOutsideCellForProximity,
@@ -3812,13 +4034,14 @@ int32 CutWithPlanarCells(
 )
 {
 	TArray<int32> TransformIndices { TransformIdx };
-	return CutMultipleWithPlanarCells(Cells, Source, TransformIndices, TransformCells, bIncludeOutsideCellInOutput, CheckDistanceAcrossOutsideCellForProximity, bSetDefaultInternalMaterialsFromCollection, VertexInterpolate);
+	return CutMultipleWithPlanarCells(Cells, Source, TransformIndices, Grout, TransformCells, bIncludeOutsideCellInOutput, CheckDistanceAcrossOutsideCellForProximity, bSetDefaultInternalMaterialsFromCollection, VertexInterpolate);
 }
 
 int32 CutMultipleWithPlanarCells_MeshBooleanPath(
 	FPlanarCells& Cells,
 	FGeometryCollection& Source,
 	const TArrayView<const int32>& TransformIndices,
+	double Grout,
 	const TOptional<FTransform>& TransformCells,
 	bool bIncludeOutsideCellInOutput,
 	float CheckDistanceAcrossOutsideCellForProximity,
@@ -3831,6 +4054,7 @@ int32 CutMultipleWithPlanarCells(
 	FPlanarCells &Cells,
 	FGeometryCollection& Source,
 	const TArrayView<const int32>& TransformIndices,
+	double Grout,
 	const TOptional<FTransform>& TransformCells,
 	bool bIncludeOutsideCellInOutput,
 	float CheckDistanceAcrossOutsideCellForProximity,
@@ -3844,6 +4068,7 @@ int32 CutMultipleWithPlanarCells(
 			Cells,
 			Source,
 			TransformIndices,
+			Grout,
 			TransformCells,
 			bIncludeOutsideCellInOutput,
 			CheckDistanceAcrossOutsideCellForProximity,
@@ -4248,6 +4473,7 @@ int32 CutMultipleWithMultiplePlanes(
 	FInternalSurfaceMaterials& InternalSurfaceMaterials,
 	FGeometryCollection& Collection,
 	const TArrayView<const int32>& TransformIndices,
+	double Grout,
 	const TOptional<FTransform>& TransformCells,
 	bool bSetDefaultInternalMaterialsFromCollection,
 	TFunction<void(const FGeometryCollection&, int32, const FGeometryCollection&, int32, float, int32, FGeometryCollection&)> VertexInterpolate
@@ -4278,7 +4504,7 @@ int32 CutMultipleWithMultiplePlanes(
 	double OnePercentExtend = MeshCollection.Bounds.MaxDim() * .01;
 
 	int32 NewGeomStartIdx = -1;
-	NewGeomStartIdx = MeshCollection.CutWithMultiplePlanes(Planes, &Collection, InternalSurfaceMaterials, bSetDefaultInternalMaterialsFromCollection);
+	NewGeomStartIdx = MeshCollection.CutWithMultiplePlanes(Planes, Grout, &Collection, InternalSurfaceMaterials, bSetDefaultInternalMaterialsFromCollection);
 
 	Collection.ReindexMaterials();
 	return NewGeomStartIdx;
@@ -4290,6 +4516,7 @@ int32 CutMultipleWithPlanarCells_MeshBooleanPath(
 	FPlanarCells& Cells,
 	FGeometryCollection& Source,
 	const TArrayView<const int32>& TransformIndices,
+	double Grout,
 	const TOptional<FTransform>& TransformCells,
 	bool bIncludeOutsideCellInOutput,
 	float CheckDistanceAcrossOutsideCellForProximity,
@@ -4312,11 +4539,11 @@ int32 CutMultipleWithPlanarCells_MeshBooleanPath(
 
 	FDynamicMeshCollection MeshCollection(&Source, TransformIndices, CellsToWorld);
 	double OnePercentExtend = MeshCollection.Bounds.MaxDim() * .01;
-	FCellMeshes CellMeshes(Cells, MeshCollection.Bounds, OnePercentExtend, bIncludeOutsideCellInOutput);
+	FCellMeshes CellMeshes(Cells, MeshCollection.Bounds, Grout, OnePercentExtend, bIncludeOutsideCellInOutput);
 
 	int32 NewGeomStartIdx = -1;
 
-	NewGeomStartIdx = MeshCollection.CutWithCellMeshes(Cells, CellMeshes, &Source, bSetDefaultInternalMaterialsFromCollection);
+	NewGeomStartIdx = MeshCollection.CutWithCellMeshes(Cells.InternalSurfaceMaterials, Cells.PlaneCells, CellMeshes, &Source, bSetDefaultInternalMaterialsFromCollection);
 
 	Source.ReindexMaterials();
 	return NewGeomStartIdx;
diff --git a/Engine/Plugins/Experimental/PlanarCutPlugin/Source/PlanarCut/Public/PlanarCut.h b/Engine/Plugins/Experimental/PlanarCutPlugin/Source/PlanarCut/Public/PlanarCut.h
--- a/Engine/Plugins/Experimental/PlanarCutPlugin/Source/PlanarCut/Public/PlanarCut.h
+++ b/Engine/Plugins/Experimental/PlanarCutPlugin/Source/PlanarCut/Public/PlanarCut.h
@@ -166,6 +166,7 @@ int32 PLANARCUT_API CutWithPlanarCells(
 	FPlanarCells &Cells,
 	FGeometryCollection& Collection,
 	int32 TransformIdx,
+	double Grout,
 	const TOptional<FTransform>& TransformCells = TOptional<FTransform>(),
 	bool bIncludeOutsideCellInOutput = true,
 	float CheckDistanceAcrossOutsideCellForProximity = 0,
@@ -189,6 +190,7 @@ int32 PLANARCUT_API CutMultipleWithPlanarCells(
 	FPlanarCells &Cells,
 	FGeometryCollection& Collection,
 	const TArrayView<const int32>& TransformIndices,
+	double Grout,
 	const TOptional<FTransform>& TransformCells = TOptional<FTransform>(),
 	bool bIncludeOutsideCellInOutput = true,
 	float CheckDistanceAcrossOutsideCellForProximity = 0,  // TODO: < this param does nothing in the new mode; is only needed in special cases that aren't possible in the UI currently
@@ -212,6 +214,7 @@ int32 PLANARCUT_API CutMultipleWithMultiplePlanes(
 	FInternalSurfaceMaterials& InternalSurfaceMaterials,
 	FGeometryCollection& Collection,
 	const TArrayView<const int32>& TransformIndices,
+	double Grout,
 	const TOptional<FTransform>& TransformCells = TOptional<FTransform>(),
 	bool bSetDefaultInternalMaterialsFromCollection = true,
 	TFunction<void(const FGeometryCollection&, int32, const FGeometryCollection&, int32, float, int32, FGeometryCollection&)> VertexInterpolate = DefaultVertexInterpolation
