Use two full 3-vector for tangent space in fracture, and also fix transforms on geometrycollection->dynamicmesh conversion

#rnx

diff --git a/Engine/Plugins/Experimental/PlanarCutPlugin/Source/PlanarCut/Private/PlanarCut.cpp b/Engine/Plugins/Experimental/PlanarCutPlugin/Source/PlanarCut/Private/PlanarCut.cpp
--- a/Engine/Plugins/Experimental/PlanarCutPlugin/Source/PlanarCut/Private/PlanarCut.cpp
+++ b/Engine/Plugins/Experimental/PlanarCutPlugin/Source/PlanarCut/Private/PlanarCut.cpp
@@ -126,7 +126,7 @@ namespace UE
 				Mesh.Attributes()->EnableMaterialID();
 				ensure(Mesh.Attributes()->NumAttachedAttributes() == 0);
 				Mesh.Attributes()->AttachAttribute(TangentUAttribName, new TDynamicMeshVertexAttribute<float, 3>(&Mesh));
-				Mesh.Attributes()->AttachAttribute(TangentVAttribName, new TDynamicMeshVertexAttribute<float, 1>(&Mesh)); // just a sign
+				Mesh.Attributes()->AttachAttribute(TangentVAttribName, new TDynamicMeshVertexAttribute<float, 3>(&Mesh)); // just a sign
 				TDynamicMeshScalarTriangleAttribute<bool>* VisAttrib = new TDynamicMeshScalarTriangleAttribute<bool>(&Mesh);
 				VisAttrib->Initialize(true);
 				Mesh.Attributes()->AttachAttribute(VisibleAttribName, VisAttrib);
@@ -151,17 +151,16 @@ namespace UE
 				checkSlow(IsAugmented(Mesh));
 				TDynamicMeshVertexAttribute<float, 3>* Us =
 					static_cast<TDynamicMeshVertexAttribute<float, 3>*>(Mesh.Attributes()->GetAttachedAttribute(TangentUAttribName));
-				TDynamicMeshVertexAttribute<float, 1>* VSigns =
-					static_cast<TDynamicMeshVertexAttribute<float, 1>*>(Mesh.Attributes()->GetAttachedAttribute(TangentVAttribName));
+				TDynamicMeshVertexAttribute<float, 3>* Vs =
+					static_cast<TDynamicMeshVertexAttribute<float, 3>*>(Mesh.Attributes()->GetAttachedAttribute(TangentVAttribName));
 
 				for (int VID : Mesh.VertexIndicesItr())
 				{
 					FVector3f N = Mesh.GetVertexNormal(VID);
-					FVector3f U;
-					VectorUtil::MakePerpVector(N, U);
+					FVector3f U, V;
+					VectorUtil::MakePerpVectors(N, U, V);
 					Us->SetValue(VID, U);
-					float VSign = 1;
-					VSigns->SetValue(VID, &VSign);
+					Vs->SetValue(VID, V);
 				}
 
 				TDynamicMeshScalarTriangleAttribute<bool>* Visible =
@@ -193,34 +192,10 @@ namespace UE
 				checkSlow(IsAugmented(Mesh));
 				TDynamicMeshVertexAttribute<float, 3>* Us =
 					static_cast<TDynamicMeshVertexAttribute<float, 3>*>(Mesh.Attributes()->GetAttachedAttribute(TangentUAttribName));
-				TDynamicMeshVertexAttribute<float, 1>* VSigns =
-					static_cast<TDynamicMeshVertexAttribute<float, 1>*>(Mesh.Attributes()->GetAttachedAttribute(TangentVAttribName));
-				float VSign = FMathf::SignNonZero(Normal.Cross(TangentU).Dot(TangentV));
+				TDynamicMeshVertexAttribute<float, 3>* Vs =
+					static_cast<TDynamicMeshVertexAttribute<float, 3>*>(Mesh.Attributes()->GetAttachedAttribute(TangentVAttribName));
 				Us->SetValue(VID, TangentU);
-				VSigns->SetValue(VID, &VSign);
-			}
-
-			void FixTangents(FDynamicMesh3& Mesh)
-			{
-				checkSlow(IsAugmented(Mesh));
-				TDynamicMeshVertexAttribute<float, 3>* Us =
-					static_cast<TDynamicMeshVertexAttribute<float, 3>*>(Mesh.Attributes()->GetAttachedAttribute(TangentUAttribName));
-				TDynamicMeshVertexAttribute<float, 1>* VSigns =
-					static_cast<TDynamicMeshVertexAttribute<float, 1>*>(Mesh.Attributes()->GetAttachedAttribute(TangentVAttribName));
-
-				for (int VID : Mesh.VertexIndicesItr())
-				{
-					FVector3f N = Mesh.GetVertexNormal(VID);
-					FVector3f U;
-					Us->GetValue(VID, U);
-					U = N.Cross(U).Cross(N);
-					U.Normalize();
-					Us->SetValue(VID, U);
-					float VSign;
-					VSigns->GetValue(VID, &VSign);
-					VSign = FMathf::SignNonZero(VSign);
-					VSigns->SetValue(VID, &VSign);
-				}
+				Vs->SetValue(VID, TangentV);
 			}
 
 			void GetTangent(const FDynamicMesh3& Mesh, int VID, FVector3f& U, FVector3f& V)
@@ -228,13 +203,11 @@ namespace UE
 				checkSlow(IsAugmented(Mesh));
 				TDynamicMeshVertexAttribute<float, 3>* Us =
 					static_cast<TDynamicMeshVertexAttribute<float, 3>*>(Mesh.Attributes()->GetAttachedAttribute(TangentUAttribName));
-				TDynamicMeshVertexAttribute<float, 1>* VSigns =
-					static_cast<TDynamicMeshVertexAttribute<float, 1>*>(Mesh.Attributes()->GetAttachedAttribute(TangentVAttribName));
+				TDynamicMeshVertexAttribute<float, 3>* Vs =
+					static_cast<TDynamicMeshVertexAttribute<float, 3>*>(Mesh.Attributes()->GetAttachedAttribute(TangentVAttribName));
 				FVector3f Normal = Mesh.GetVertexNormal(VID);
 				Us->GetValue(VID, U);
-				float VSign;
-				VSigns->GetValue(VID, &VSign);
-				V = Normal.Cross(U) * FMathf::SignNonZero(VSign);
+				Vs->GetValue(VID, V);
 			}
 		}
 	}
@@ -1232,9 +1205,11 @@ struct FDynamicMeshCollection
 				VertexInfo.Position = CollectionToLocal.TransformPosition(FVector3d(Collection->Vertex[Idx]));
 				VertexInfo.UV = FVector2f(Collection->UV[Idx]);
 				VertexInfo.Color = FVector3f(Collection->Color[Idx]);
-				VertexInfo.Normal = CollectionToLocal.TransformNormal(FVector3f(Collection->Normal[Idx]));
+				VertexInfo.Normal = CollectionToLocal.TransformVectorNoScale(FVector3f(Collection->Normal[Idx]));
 				int VID = Mesh.AppendVertex(VertexInfo);
-				UE::PlanarCutInternals::AugmentDynamicMesh::SetTangent(Mesh, VID, VertexInfo.Normal, Collection->TangentU[Idx], Collection->TangentV[Idx]);
+				UE::PlanarCutInternals::AugmentDynamicMesh::SetTangent(Mesh, VID, VertexInfo.Normal,
+					CollectionToLocal.TransformVectorNoScale(FVector3f(Collection->TangentU[Idx])),
+					CollectionToLocal.TransformVectorNoScale(FVector3f(Collection->TangentV[Idx])));
 			}
 			FIntVector VertexOffset(VertexStart, VertexStart, VertexStart);
 			for (int32 Idx = Collection->FaceStart[GeometryIdx], N = Collection->FaceStart[GeometryIdx] + FaceCount; Idx < N; Idx++)
@@ -1575,7 +1550,6 @@ struct FDynamicMeshCollection
 			for (int32 ToCutIdx : ToCutIndices)
 			{
 				FDynamicMesh3& Mesh = ToCut[ToCutIdx]->AugMesh;
-				UE::PlanarCutInternals::AugmentDynamicMesh::FixTangents(Mesh);
 
 				int32 CreatedGeometryIdx = AppendToCollection(ToCut[ToCutIdx]->ToCollection, Mesh, ToCut[ToCutIdx]->TransformIndex, SubPartIdx++, *Collection, InternalMaterialID);
 				ToCutIdxToGeometryIdx[ToCutIdx] = CreatedGeometryIdx;
@@ -1665,7 +1639,6 @@ struct FDynamicMeshCollection
 					if (BooleanResults[CellIdx].IsValid() && BooleanResults[CellIdx]->TriangleCount() > 0)
 					{
 						FDynamicMesh3& AugBoolResult = *BooleanResults[CellIdx];
-						UE::PlanarCutInternals::AugmentDynamicMesh::FixTangents(AugBoolResult);
 						for (int TID : AugBoolResult.TriangleIndicesItr())
 						{
 							int MID = AugBoolResult.Attributes()->GetMaterialID()->GetValue(TID);
@@ -1932,7 +1905,7 @@ struct FDynamicMeshCollection
 			checkSlow(Mesh.IsVertex(VID)); // mesh is compact
 			int32 CopyToIdx = VerticesStart + VID;
 			Output.Vertex[CopyToIdx] = ToCollection.TransformPosition(FVector(Mesh.GetVertex(VID)));
-			Output.Normal[CopyToIdx] = ToCollection.TransformVectorNoScale(FVector(Mesh.GetVertexNormal(VID))); // TODO: use correct normal transform instead
+			Output.Normal[CopyToIdx] = ToCollection.TransformVectorNoScale(FVector(Mesh.GetVertexNormal(VID)));
 			Output.UV[CopyToIdx] = FVector2D(Mesh.GetVertexUV(VID));
 			FVector3f TangentU, TangentV;
 			UE::PlanarCutInternals::AugmentDynamicMesh::GetTangent(Mesh, VID, TangentU, TangentV);
