Island distance tolerance for meshes should be larger to avoid bad splits

#rnx

diff --git a/Engine/Plugins/Experimental/PlanarCutPlugin/Source/PlanarCut/Private/PlanarCut.cpp b/Engine/Plugins/Experimental/PlanarCutPlugin/Source/PlanarCut/Private/PlanarCut.cpp
--- a/Engine/Plugins/Experimental/PlanarCutPlugin/Source/PlanarCut/Private/PlanarCut.cpp
+++ b/Engine/Plugins/Experimental/PlanarCutPlugin/Source/PlanarCut/Private/PlanarCut.cpp
@@ -1798,7 +1798,7 @@ struct FDynamicMeshCollection
 	// Split mesh into connected components, including implicit connections by co-located vertices
 	bool SplitIslands(FDynamicMesh3& Source, TArray<FDynamicMesh3>& SeparatedMeshes)
 	{
-		double SnapDistance = FMathd::ZeroTolerance * 100;
+		double SnapDistance = 1e-03;
 		TPointHashGrid3d<int> VertHash(SnapDistance * 10, -1);
 		FDisjointSet VertComponents(Source.MaxVertexID());
 		// Add Source vertices to hash & disjoint sets
