add collision sample spacing setting to fracture tools, so user can set a target density of collision samples and they will be automatically added on meshes where vertices are too far apart
add a "resample" tool to the fracture tool set so users can update collision sample spacing w/out redoing fracture, and also to help visualize the point placement

diff --git a/Engine/Plugins/Experimental/ChaosEditor/Source/FractureEditor/Private/FractureEditorCommands.h b/Engine/Plugins/Experimental/ChaosEditor/Source/FractureEditor/Private/FractureEditorCommands.h
--- a/Engine/Plugins/Experimental/ChaosEditor/Source/FractureEditor/Private/FractureEditorCommands.h
+++ b/Engine/Plugins/Experimental/ChaosEditor/Source/FractureEditor/Private/FractureEditorCommands.h
@@ -48,6 +48,7 @@ class FFractureEditorCommands : public TCommands<FFractureEditorCommands>
 		TSharedPtr< FUICommandInfo > Slice;
 		TSharedPtr< FUICommandInfo > Brick;
 		TSharedPtr< FUICommandInfo > Texture;
+		TSharedPtr< FUICommandInfo > Resample;
 
 		// Property Commands
 		TSharedPtr< FUICommandInfo > SetInitialDynamicState;
diff --git a/Engine/Plugins/Experimental/ChaosEditor/Source/FractureEditor/Private/FractureEditorModeToolkit.cpp b/Engine/Plugins/Experimental/ChaosEditor/Source/FractureEditor/Private/FractureEditorModeToolkit.cpp
--- a/Engine/Plugins/Experimental/ChaosEditor/Source/FractureEditor/Private/FractureEditorModeToolkit.cpp
+++ b/Engine/Plugins/Experimental/ChaosEditor/Source/FractureEditor/Private/FractureEditorModeToolkit.cpp
@@ -500,6 +500,7 @@ void FFractureEditorModeToolkit::BuildToolPalette(FName PaletteIndex, class FToo
 		ToolbarBuilder.AddToolBarButton(Commands.Planar);
 		ToolbarBuilder.AddToolBarButton(Commands.Slice);
 		ToolbarBuilder.AddToolBarButton(Commands.Brick);
+		ToolbarBuilder.AddToolBarButton(Commands.Resample);
 	}
 	else if (PaletteIndex == TEXT("Cluster"))
 	{
diff --git a/Engine/Plugins/Experimental/ChaosEditor/Source/FractureEditor/Private/FractureToolBrick.cpp b/Engine/Plugins/Experimental/ChaosEditor/Source/FractureEditor/Private/FractureToolBrick.cpp
--- a/Engine/Plugins/Experimental/ChaosEditor/Source/FractureEditor/Private/FractureToolBrick.cpp
+++ b/Engine/Plugins/Experimental/ChaosEditor/Source/FractureEditor/Private/FractureToolBrick.cpp
@@ -48,6 +48,7 @@ TArray<UObject*> UFractureToolBrick::GetSettingsObjects() const
 { 
 	TArray<UObject*> Settings; 
 	Settings.Add(CutterSettings);
+	Settings.Add(CollisionSettings);
 	Settings.Add(BrickSettings);
 	return Settings;
 }
@@ -319,7 +320,7 @@ int32 UFractureToolBrick::ExecuteFracture(const FFractureToolContext& FractureCo
 			VoronoiPlanarCells.InternalSurfaceMaterials.NoiseSettings = NoiseSettings;
 		}
 
-		return CutMultipleWithPlanarCells(VoronoiPlanarCells, *FractureContext.GetGeometryCollection(), FractureContext.GetSelection(), 0);
+		return CutMultipleWithPlanarCells(VoronoiPlanarCells, *FractureContext.GetGeometryCollection(), FractureContext.GetSelection(), 0, CollisionSettings->PointSpacing);
 	}
 
 	return INDEX_NONE;
diff --git a/Engine/Plugins/Experimental/ChaosEditor/Source/FractureEditor/Private/FractureToolClusterCutter.cpp b/Engine/Plugins/Experimental/ChaosEditor/Source/FractureEditor/Private/FractureToolClusterCutter.cpp
--- a/Engine/Plugins/Experimental/ChaosEditor/Source/FractureEditor/Private/FractureToolClusterCutter.cpp
+++ b/Engine/Plugins/Experimental/ChaosEditor/Source/FractureEditor/Private/FractureToolClusterCutter.cpp
@@ -40,6 +40,7 @@ TArray<UObject*> UFractureToolClusterCutter::GetSettingsObjects() const
 {
 	TArray<UObject*> ReturnSettings;
 	ReturnSettings.Add(CutterSettings);
+	ReturnSettings.Add(CollisionSettings);
 	ReturnSettings.Add(ClusterSettings);
 	return ReturnSettings;
 }
diff --git a/Engine/Plugins/Experimental/ChaosEditor/Source/FractureEditor/Private/FractureToolCutter.cpp b/Engine/Plugins/Experimental/ChaosEditor/Source/FractureEditor/Private/FractureToolCutter.cpp
--- a/Engine/Plugins/Experimental/ChaosEditor/Source/FractureEditor/Private/FractureToolCutter.cpp
+++ b/Engine/Plugins/Experimental/ChaosEditor/Source/FractureEditor/Private/FractureToolCutter.cpp
@@ -19,6 +19,8 @@ UFractureToolCutterBase::UFractureToolCutterBase(const FObjectInitializer& ObjIn
 {
 	CutterSettings = NewObject<UFractureCutterSettings>(GetTransientPackage(), UFractureCutterSettings::StaticClass());
 	CutterSettings->OwnerTool = this;
+	CollisionSettings = NewObject<UFractureCollisionSettings>(GetTransientPackage(), UFractureCollisionSettings::StaticClass());
+	CollisionSettings->OwnerTool = this;
 }
 
 bool UFractureToolCutterBase::CanExecute() const
@@ -177,7 +179,7 @@ int32 UFractureToolVoronoiCutterBase::ExecuteFracture(const FFractureToolContext
 			VoronoiPlanarCells.InternalSurfaceMaterials.NoiseSettings = NoiseSettings;
 		}
 
-		return CutMultipleWithPlanarCells(VoronoiPlanarCells, *(FractureContext.GetGeometryCollection()), FractureContext.GetSelection(), CutterSettings->Grout);
+		return CutMultipleWithPlanarCells(VoronoiPlanarCells, *(FractureContext.GetGeometryCollection()), FractureContext.GetSelection(), CutterSettings->Grout, CollisionSettings->PointSpacing);
 	}
 
 	return INDEX_NONE;
diff --git a/Engine/Plugins/Experimental/ChaosEditor/Source/FractureEditor/Private/FractureToolCutter.h b/Engine/Plugins/Experimental/ChaosEditor/Source/FractureEditor/Private/FractureToolCutter.h
--- a/Engine/Plugins/Experimental/ChaosEditor/Source/FractureEditor/Private/FractureToolCutter.h
+++ b/Engine/Plugins/Experimental/ChaosEditor/Source/FractureEditor/Private/FractureToolCutter.h
@@ -66,11 +66,26 @@ public:
 	int32 OctaveNumber;
 
 	/** Spacing between vertices on cut surfaces, where noise is added.  Larger spacing between vertices will create more efficient meshes with fewer triangles, but less resolution to see the shape of the added noise  */
-	UPROPERTY(EditAnywhere, Category = Noise, meta = (UIMin = "1", ClampMin = "0.1"))
+	UPROPERTY(EditAnywhere, Category = Noise, meta = (DisplayName = "Point Spacing", UIMin = "1", ClampMin = "0.1"))
 	float SurfaceResolution;
-
 };
 
+/** Settings related to the collision properties of the fractured mesh pieces */
+UCLASS(config = EditorPerProjectUserSettings)
+class UFractureCollisionSettings : public UFractureToolSettings
+{
+public:
+	GENERATED_BODY()
+
+	UFractureCollisionSettings(const FObjectInitializer& ObjInit)
+	: Super(ObjInit) {}
+
+	/** Target spacing between collision samples on the mesh surface. */
+	UPROPERTY(EditAnywhere, Category = Collision, meta = (UIMin = "1", ClampMin = "0.1"))
+	float PointSpacing = 50.0f;
+
+	// TODO: add remeshing options here as well
+};
 
 UCLASS(Abstract, DisplayName = "Cutter Base", Category = "FractureTools")
 class UFractureToolCutterBase : public UFractureInteractiveTool
@@ -91,6 +106,9 @@ protected:
 	UPROPERTY(EditAnywhere, Category = Slicing)
 	TObjectPtr<UFractureCutterSettings> CutterSettings;
 
+	UPROPERTY(EditAnywhere, Category = Collision)
+	TObjectPtr<UFractureCollisionSettings> CollisionSettings;
+
 };
 
 
diff --git a/Engine/Plugins/Experimental/ChaosEditor/Source/FractureEditor/Private/FractureToolPlaneCut.cpp b/Engine/Plugins/Experimental/ChaosEditor/Source/FractureEditor/Private/FractureToolPlaneCut.cpp
--- a/Engine/Plugins/Experimental/ChaosEditor/Source/FractureEditor/Private/FractureToolPlaneCut.cpp
+++ b/Engine/Plugins/Experimental/ChaosEditor/Source/FractureEditor/Private/FractureToolPlaneCut.cpp
@@ -73,6 +73,7 @@ TArray<UObject*> UFractureToolPlaneCut::GetSettingsObjects() const
  {
 	TArray<UObject*> Settings;
 	Settings.Add(CutterSettings);
+	Settings.Add(CollisionSettings);
 	Settings.Add(PlaneCutSettings);
 	return Settings;
 }
@@ -130,7 +131,7 @@ int32 UFractureToolPlaneCut::ExecuteFracture(const FFractureToolContext& Fractur
 			InternalSurfaceMaterials.NoiseSettings = NoiseSettings;
 		}
 
-		return CutMultipleWithMultiplePlanes(CuttingPlanes, InternalSurfaceMaterials, *FractureContext.GetGeometryCollection(), FractureContext.GetSelection(), CutterSettings->Grout);
+		return CutMultipleWithMultiplePlanes(CuttingPlanes, InternalSurfaceMaterials, *FractureContext.GetGeometryCollection(), FractureContext.GetSelection(), CutterSettings->Grout, CollisionSettings->PointSpacing);
 	}
 
 	return INDEX_NONE;
diff --git a/Engine/Plugins/Experimental/ChaosEditor/Source/FractureEditor/Private/FractureToolRadial.cpp b/Engine/Plugins/Experimental/ChaosEditor/Source/FractureEditor/Private/FractureToolRadial.cpp
--- a/Engine/Plugins/Experimental/ChaosEditor/Source/FractureEditor/Private/FractureToolRadial.cpp
+++ b/Engine/Plugins/Experimental/ChaosEditor/Source/FractureEditor/Private/FractureToolRadial.cpp
@@ -40,6 +40,7 @@ TArray<UObject*> UFractureToolRadial::GetSettingsObjects() const
 { 
 	TArray<UObject*> Settings; 
 	Settings.Add(CutterSettings);
+	Settings.Add(CollisionSettings);
 	Settings.Add(RadialSettings);
 	return Settings;
 }
diff --git a/Engine/Plugins/Experimental/ChaosEditor/Source/FractureEditor/Private/FractureToolResample.cpp b/Engine/Plugins/Experimental/ChaosEditor/Source/FractureEditor/Private/FractureToolResample.cpp
new file mode 100644
--- /dev/null
+++ b/Engine/Plugins/Experimental/ChaosEditor/Source/FractureEditor/Private/FractureToolResample.cpp
@@ -0,0 +1,103 @@
+// Copyright Epic Games, Inc. All Rights Reserved.
+
+#include "FractureToolResample.h"
+#include "FractureEditorStyle.h"
+#include "PlanarCut.h"
+#include "GeometryCollection/GeometryCollection.h"
+#include "GeometryCollection/GeometryCollectionObject.h"
+#include "GeometryCollection/GeometryCollectionAlgo.h"
+#include "FractureEditorModeToolkit.h"
+#include "FractureToolContext.h"
+
+#define LOCTEXT_NAMESPACE "FractureResample"
+
+
+UFractureToolResample::UFractureToolResample(const FObjectInitializer& ObjInit) 
+	: Super(ObjInit) 
+{
+}
+
+FText UFractureToolResample::GetDisplayText() const
+{
+	return FText(NSLOCTEXT("Resample", "FractureToolResample", "Update Collision Samples")); 
+}
+
+FText UFractureToolResample::GetTooltipText() const 
+{
+	return FText(NSLOCTEXT("Resample", "FractureToolResampleTooltip", "The Resample tool can add collision samples in large flat regions that otherwise might have poor collision response."));
+}
+
+FSlateIcon UFractureToolResample::GetToolIcon() const 
+{
+	return FSlateIcon("FractureEditorStyle", "FractureEditor.Mesh");
+}
+
+void UFractureToolResample::RegisterUICommand( FFractureEditorCommands* BindingContext ) 
+{
+	UI_COMMAND_EXT( BindingContext, UICommandInfo, "Resample", "Resample", "Resample", EUserInterfaceActionType::ToggleButton, FInputChord() );
+	BindingContext->Resample = UICommandInfo;
+}
+
+void UFractureToolResample::Render(const FSceneView* View, FViewport* Viewport, FPrimitiveDrawInterface* PDI)
+{
+	for (FVector Point : GeneratedPoints)
+	{
+		PDI->DrawPoint(Point, FLinearColor::Green, 2.f, SDPG_Foreground);
+	}
+	
+}
+
+TArray<UObject*> UFractureToolResample::GetSettingsObjects() const
+ {
+	TArray<UObject*> Settings;
+	Settings.Add(CollisionSettings);
+	return Settings;
+}
+
+
+void UFractureToolResample::FractureContextChanged()
+{
+	TArray<FFractureToolContext> FractureContexts = GetFractureToolContexts();
+
+	GeneratedPoints.Reset();
+	for (const FFractureToolContext& FractureContext : FractureContexts)
+	{
+		FGeometryCollection& Collection = *FractureContext.GetGeometryCollection();
+		
+
+		FTransform OuterTransform = FractureContext.GetTransform();
+		for (int32 TransformIdx : FractureContext.GetSelection())
+		{
+			FTransform InnerTransform = GeometryCollectionAlgo::GlobalMatrix(Collection.Transform, Collection.Parent, TransformIdx);
+			if (Collection.HasAttribute("ExplodedVector", FGeometryCollection::TransformGroup))
+			{
+				TManagedArray<FVector>& ExplodedVectors = Collection.GetAttribute<FVector>("ExplodedVector", FGeometryCollection::TransformGroup);
+				InnerTransform = InnerTransform * FTransform(ExplodedVectors[TransformIdx]);
+			}
+			
+			FTransform CombinedTransform = InnerTransform * OuterTransform;
+			int32 GeometryIdx = Collection.TransformToGeometryIndex[TransformIdx];
+			int32 VertStart = Collection.VertexStart[GeometryIdx];
+			int32 VertCount = Collection.VertexCount[GeometryIdx];
+			for (int32 VIdx = VertStart; VIdx < VertStart + VertCount; VIdx++)
+			{
+				GeneratedPoints.Add(CombinedTransform.TransformPosition(Collection.Vertex[VIdx]));
+			}
+		}
+	}
+}
+
+int32 UFractureToolResample::ExecuteFracture(const FFractureToolContext& FractureContext)
+{
+	if (FractureContext.IsValid())
+	{
+		const UFractureCollisionSettings* LocalCutSettings = CollisionSettings;
+
+		return AddCollisionSampleVertices(LocalCutSettings->PointSpacing, *FractureContext.GetGeometryCollection(), FractureContext.GetSelection());
+	}
+
+	return INDEX_NONE;
+}
+
+
+#undef LOCTEXT_NAMESPACE
\ No newline at end of file
diff --git a/Engine/Plugins/Experimental/ChaosEditor/Source/FractureEditor/Private/FractureToolResample.h b/Engine/Plugins/Experimental/ChaosEditor/Source/FractureEditor/Private/FractureToolResample.h
new file mode 100644
--- /dev/null
+++ b/Engine/Plugins/Experimental/ChaosEditor/Source/FractureEditor/Private/FractureToolResample.h
@@ -0,0 +1,44 @@
+// Copyright Epic Games, Inc. All Rights Reserved.
+
+#pragma once
+
+#include "FractureToolCutter.h"
+#include "Engine/StaticMeshActor.h"
+
+#include "FractureToolResample.generated.h"
+
+class FFractureToolContext;
+
+// Note this tool doesn't actually fracture, but it does remake pieces of geometry and shares a lot of machinery with the fracture tools
+UCLASS(DisplayName = "Resample Tool", Category = "FractureTools")
+class UFractureToolResample : public UFractureToolCutterBase
+{
+public:
+	GENERATED_BODY()
+
+	UFractureToolResample(const FObjectInitializer& ObjInit);
+
+	// UFractureTool Interface
+	virtual FText GetDisplayText() const override;
+	virtual FText GetTooltipText() const override;
+	virtual FSlateIcon GetToolIcon() const override;
+
+	virtual FText GetApplyText() const override { return FText(NSLOCTEXT("Resample", "ExecuteResample", "Resample")); }
+
+	void Render(const FSceneView* View, FViewport* Viewport, FPrimitiveDrawInterface* PDI) override;
+
+	virtual void RegisterUICommand(FFractureEditorCommands* BindingContext) override;
+
+	virtual TArray<UObject*> GetSettingsObjects() const override;
+
+	virtual void FractureContextChanged() override;
+	virtual int32 ExecuteFracture(const FFractureToolContext& FractureContext) override;
+
+private:
+
+
+	TArray<FVector> GeneratedPoints;
+
+};
+
+
diff --git a/Engine/Plugins/Experimental/ChaosEditor/Source/FractureEditor/Private/FractureToolSlice.cpp b/Engine/Plugins/Experimental/ChaosEditor/Source/FractureEditor/Private/FractureToolSlice.cpp
--- a/Engine/Plugins/Experimental/ChaosEditor/Source/FractureEditor/Private/FractureToolSlice.cpp
+++ b/Engine/Plugins/Experimental/ChaosEditor/Source/FractureEditor/Private/FractureToolSlice.cpp
@@ -48,6 +48,7 @@ TArray<UObject*> UFractureToolSlice::GetSettingsObjects() const
 { 
 	TArray<UObject*> Settings; 
 	Settings.Add(CutterSettings);
+	Settings.Add(CollisionSettings);
 	Settings.Add(SliceSettings);
 	return Settings;
 }
@@ -169,7 +170,7 @@ int32 UFractureToolSlice::ExecuteFracture(const FFractureToolContext& FractureCo
 			InternalSurfaceMaterials.NoiseSettings = NoiseSettings;
 		}
 
-		return CutMultipleWithMultiplePlanes(CuttingPlanes, InternalSurfaceMaterials, *FractureContext.GetGeometryCollection(), FractureContext.GetSelection(), CutterSettings->Grout);
+		return CutMultipleWithMultiplePlanes(CuttingPlanes, InternalSurfaceMaterials, *FractureContext.GetGeometryCollection(), FractureContext.GetSelection(), CutterSettings->Grout, CollisionSettings->PointSpacing);
 	}
 
 	return INDEX_NONE;
diff --git a/Engine/Plugins/Experimental/GeometryProcessing/Source/GeometricObjects/Public/Spatial/PointHashGrid3.h b/Engine/Plugins/Experimental/GeometryProcessing/Source/GeometricObjects/Public/Spatial/PointHashGrid3.h
--- a/Engine/Plugins/Experimental/GeometryProcessing/Source/GeometricObjects/Public/Spatial/PointHashGrid3.h
+++ b/Engine/Plugins/Experimental/GeometryProcessing/Source/GeometricObjects/Public/Spatial/PointHashGrid3.h
@@ -100,6 +100,35 @@ public:
 	}
 
 
+	/**
+	 * Test if the cell containing Position is empty. This function is thread-safe.
+	 * Can be used to skip a more expensive range search, in some cases.
+	 * 
+	 * @return true if the cell containing Position is empty
+	 */
+	bool IsCellEmpty(const FVector3<RealType>& Position)
+	{
+		FVector3i Idx = Indexer.ToGrid(Position);
+		{
+			FScopeLock Lock(&CriticalSection);
+			return !Hash.Contains(Idx);
+		}
+	}
+
+
+	/**
+	 * Test if the whole cell containing Position is empty, without locking / thread-safety
+	 * Can be used to skip a more expensive range search, in some cases.
+	 * 
+	 * @return true if the cell containing Position is empty
+	 */
+	bool IsCellEmptyUnsafe(const FVector3<RealType>& Position)
+	{
+		FVector3i Idx = Indexer.ToGrid(Position);
+		return !Hash.Contains(Idx);
+	}
+
+
 	/**
 	 * Move value from old to new position. This function is thread-safe.
 	 * @param Value the point/value to update
diff --git a/Engine/Plugins/Experimental/PlanarCutPlugin/Source/PlanarCut/Private/PlanarCut.cpp b/Engine/Plugins/Experimental/PlanarCutPlugin/Source/PlanarCut/Private/PlanarCut.cpp
--- a/Engine/Plugins/Experimental/PlanarCutPlugin/Source/PlanarCut/Private/PlanarCut.cpp
+++ b/Engine/Plugins/Experimental/PlanarCutPlugin/Source/PlanarCut/Private/PlanarCut.cpp
@@ -19,6 +19,7 @@
 #include "DynamicMesh3.h"
 #include "DynamicMeshEditor.h"
 #include "DynamicMeshAABBTree3.h"
+#include "Selections/MeshConnectedComponents.h"
 #include "MeshTransforms.h"
 #include "Operations/MeshBoolean.h"
 #include "Operations/MeshSelfUnion.h"
@@ -126,7 +127,7 @@ namespace UE
 				Mesh.Attributes()->EnableMaterialID();
 				ensure(Mesh.Attributes()->NumAttachedAttributes() == 0);
 				Mesh.Attributes()->AttachAttribute(TangentUAttribName, new TDynamicMeshVertexAttribute<float, 3>(&Mesh));
-				Mesh.Attributes()->AttachAttribute(TangentVAttribName, new TDynamicMeshVertexAttribute<float, 3>(&Mesh)); // just a sign
+				Mesh.Attributes()->AttachAttribute(TangentVAttribName, new TDynamicMeshVertexAttribute<float, 3>(&Mesh));
 				TDynamicMeshScalarTriangleAttribute<bool>* VisAttrib = new TDynamicMeshScalarTriangleAttribute<bool>(&Mesh);
 				VisAttrib->Initialize(true);
 				Mesh.Attributes()->AttachAttribute(VisibleAttribName, VisAttrib);
@@ -209,6 +210,121 @@ namespace UE
 				Us->GetValue(VID, U);
 				Vs->GetValue(VID, V);
 			}
+
+			// per component sampling is a rough heuristic to avoid doing geodesic distance but still get points on a 'thin' slice
+			void AddCollisionSamplesPerComponent(FDynamicMesh3& Mesh, double Spacing)
+			{
+				checkSlow(IsAugmented(Mesh));
+				FMeshConnectedComponents Components(&Mesh);
+				// TODO: if/when we switch to merged edges representation, pass a predicate here based on whether there's a normal seam ?
+				Components.FindConnectedTriangles();
+				TArray<TPointHashGrid3d<int>> KnownSamples;  KnownSamples.Reserve(Components.Num());
+				for (int ComponentIdx = 0; ComponentIdx < Components.Num(); ComponentIdx++)
+				{
+					KnownSamples.Emplace(.5 * Spacing / FMathd::InvSqrt3, -1);
+				}
+
+				TArray<int> AlreadySeen; AlreadySeen.Init(-1, Mesh.MaxVertexID());
+				for (int ComponentIdx = 0; ComponentIdx < Components.Num(); ComponentIdx++)
+				{
+					FMeshConnectedComponents::FComponent& Component = Components.GetComponent(ComponentIdx);
+					for (int TID : Component.Indices)
+					{
+						FIndex3i Tri = Mesh.GetTriangle(TID);
+						for (int SubIdx = 0; SubIdx < 3; SubIdx++)
+						{
+							int VID = Tri[SubIdx];
+							if (AlreadySeen[VID] != ComponentIdx)
+							{
+								AlreadySeen[VID] = ComponentIdx;
+								KnownSamples[ComponentIdx].InsertPointUnsafe(VID, Mesh.GetVertex(VID));
+							}
+						}
+					}
+				}
+				AlreadySeen.Empty();
+
+				double SpacingThreshSq = Spacing * Spacing; // if points are more than Spacing apart, consider adding a new point between them
+				for (int ComponentIdx = 0; ComponentIdx < Components.Num(); ComponentIdx++)
+				{
+					FMeshConnectedComponents::FComponent& Component = Components.GetComponent(ComponentIdx);
+					for (int TID : Component.Indices)
+					{
+						FIndex3i TriVIDs = Mesh.GetTriangle(TID);
+						FTriangle3d Triangle;
+						Mesh.GetTriVertices(TID, Triangle.V[0], Triangle.V[1], Triangle.V[2]);
+						double EdgeLensSq[3];
+						int MaxEdgeIdx = 0;
+						double MaxEdgeLenSq = 0;
+						for (int i = 2, j = 0; j < 3; i = j++)
+						{
+							double EdgeLenSq = Triangle.V[i].DistanceSquared(Triangle.V[j]);
+							if (EdgeLenSq > MaxEdgeLenSq)
+							{
+								MaxEdgeIdx = i;
+								MaxEdgeLenSq = EdgeLenSq;
+							}
+							EdgeLensSq[i] = EdgeLenSq;
+						}
+						// if we found a too-long edge, we can try sampling the tri
+						if (MaxEdgeLenSq > SpacingThreshSq)
+						{
+							FVector3f Normal = (FVector3f)VectorUtil::Normal(Triangle.V[0], Triangle.V[1], Triangle.V[2]);
+
+							// Pick number of samples based on the longest edge
+							double LongEdgeLen = FMathd::Sqrt(MaxEdgeLenSq);
+							int Divisions = FMathd::Floor(LongEdgeLen / Spacing);
+							double Factor = 1.0 / double(Divisions + 1);
+							int SecondEdgeIdx = (MaxEdgeIdx + 1) % 3;
+							int ThirdEdgeIdx = (MaxEdgeIdx + 2) % 3;
+							// Sample along the two longest edges first, then interpolate these samples
+							int SecondLongestEdgeIdx = SecondEdgeIdx;
+							if (EdgeLensSq[SecondEdgeIdx] < EdgeLensSq[ThirdEdgeIdx])
+							{
+								SecondLongestEdgeIdx = ThirdEdgeIdx;
+							}
+							int SecondLongestSecondEdgeIdx = (SecondLongestEdgeIdx + 1) % 3;
+							for (int DivI = 0; DivI < Divisions; DivI++)
+							{
+								double Along = (DivI + 1) * Factor;
+								FVector3d E1Bary(0, 0, 0), E2Bary(0, 0, 0);
+								E1Bary[MaxEdgeIdx] = Along;
+								E1Bary[SecondEdgeIdx] = 1 - Along;
+								E2Bary[SecondLongestEdgeIdx] = 1 - Along;
+								E2Bary[SecondLongestSecondEdgeIdx] = Along;
+
+								// Choose number of samples between the two edge points based on their distance
+								double AcrossDist = Triangle.BarycentricPoint(E1Bary).Distance(Triangle.BarycentricPoint(E2Bary));
+								int DivisionsAcross = FMathd::Ceil(AcrossDist / Spacing);
+								double FactorAcross = 1.0 / double(DivisionsAcross + 1);
+								for (int DivJ = 0; DivJ < DivisionsAcross; DivJ++)
+								{
+									double AlongAcross = (DivJ + 1) * FactorAcross;
+									FVector3d Bary = FVector3d::Lerp(E1Bary, E2Bary, AlongAcross);
+									FVector3d SamplePos = Triangle.BarycentricPoint(Bary);
+									if (!KnownSamples[ComponentIdx].IsCellEmptyUnsafe(SamplePos)) // fast early out; def. have pt within radius
+									{
+										continue;
+									}
+									TPair<int, double> VIDDist = KnownSamples[ComponentIdx].FindNearestInRadius(SamplePos, Spacing * .5, [&Mesh, SamplePos](int VID)
+										{
+											return Mesh.GetVertex(VID).DistanceSquared(SamplePos);
+										});
+									// No point within radius Spacing/2 -> Add a new sample
+									if (VIDDist.Key == -1)
+									{
+										// no point within radius; can add a sample here
+										FVertexInfo Info(SamplePos, Normal);
+
+										int AddedVID = Mesh.AppendVertex(Info);
+										KnownSamples[ComponentIdx].InsertPointUnsafe(AddedVID, SamplePos);
+									}
+								}
+							}
+						}
+					}
+				}
+			}
 		}
 	}
 }
@@ -1234,6 +1350,7 @@ struct FDynamicMeshCollection
 	int32 CutWithMultiplePlanes(
 		const TArrayView<const FPlane>& Planes, 
 		double Grout,
+		double CollisionSampleSpacing,
 		FGeometryCollection* Collection,
 		FInternalSurfaceMaterials& InternalSurfaceMaterials,
 		bool bSetDefaultInternalMaterialsFromCollection
@@ -1294,7 +1411,7 @@ struct FDynamicMeshCollection
 			TArray<TPair<int32, int32>> CellConnectivity;
 			CellConnectivity.Add(TPair<int32, int32>(0, -1));
 
-			return CutWithCellMeshes(InternalSurfaceMaterials, CellConnectivity, GroutCells, Collection, bSetDefaultInternalMaterialsFromCollection);
+			return CutWithCellMeshes(InternalSurfaceMaterials, CellConnectivity, GroutCells, Collection, bSetDefaultInternalMaterialsFromCollection, CollisionSampleSpacing);
 		}
 
 		bool bHasProximity = Collection->HasAttribute("Proximity", FGeometryCollection::GeometryGroup);
@@ -1551,7 +1668,8 @@ struct FDynamicMeshCollection
 			{
 				FDynamicMesh3& Mesh = ToCut[ToCutIdx]->AugMesh;
 
-				int32 CreatedGeometryIdx = AppendToCollection(ToCut[ToCutIdx]->ToCollection, Mesh, ToCut[ToCutIdx]->TransformIndex, SubPartIdx++, *Collection, InternalMaterialID);
+				FString BoneName = GetBoneName(*Collection, ToCut[ToCutIdx]->TransformIndex, SubPartIdx++);
+				int32 CreatedGeometryIdx = AppendToCollection(ToCut[ToCutIdx]->ToCollection, Mesh, CollisionSampleSpacing, ToCut[ToCutIdx]->TransformIndex, BoneName, *Collection, InternalMaterialID);
 				ToCutIdxToGeometryIdx[ToCutIdx] = CreatedGeometryIdx;
 				if (FirstCreatedIndex == -1)
 				{
@@ -1583,7 +1701,7 @@ struct FDynamicMeshCollection
 	 * @param bSetDefaultInternalMaterialsFromCollection If true, set internal materials to the most common external material + 1, following a convenient artist convention
 	 * @return Index of the first created geometry
 	 */
-	int32 CutWithCellMeshes(const FInternalSurfaceMaterials& InternalSurfaceMaterials, const TArray<TPair<int32, int32>>& CellConnectivity, FCellMeshes& CellMeshes, FGeometryCollection* Collection, bool bSetDefaultInternalMaterialsFromCollection)
+	int32 CutWithCellMeshes(const FInternalSurfaceMaterials& InternalSurfaceMaterials, const TArray<TPair<int32, int32>>& CellConnectivity, FCellMeshes& CellMeshes, FGeometryCollection* Collection, bool bSetDefaultInternalMaterialsFromCollection, double CollisionSampleSpacing)
 	{
 		// TODO: should we do these cuts in parallel, and the appends sequentially below?
 		int32 FirstIdx = -1;
@@ -1655,7 +1773,8 @@ struct FDynamicMeshCollection
 							for (int32 i = 0; i < Islands.Num(); i++)
 							{
 								FDynamicMesh3& Island = Islands[i];
-								CreatedGeometryIdx = AppendToCollection(Surface.ToCollection, Island, Surface.TransformIndex, SubPartIndex++, *Collection, InternalMaterialID);
+								FString BoneName = GetBoneName(*Collection, Surface.TransformIndex, SubPartIndex++);
+								CreatedGeometryIdx = AppendToCollection(Surface.ToCollection, Island, CollisionSampleSpacing, Surface.TransformIndex, BoneName, *Collection, InternalMaterialID);
 								CellToGeometry.Add(CellIdx, CreatedGeometryIdx);
 								if (i > 0)
 								{
@@ -1670,7 +1789,8 @@ struct FDynamicMeshCollection
 						}
 						else
 						{
-							CreatedGeometryIdx = AppendToCollection(Surface.ToCollection, AugBoolResult, Surface.TransformIndex, SubPartIndex++, *Collection, InternalMaterialID);
+							FString BoneName = GetBoneName(*Collection, Surface.TransformIndex, SubPartIndex++);
+							CreatedGeometryIdx = AppendToCollection(Surface.ToCollection, AugBoolResult, CollisionSampleSpacing, Surface.TransformIndex, BoneName, *Collection, InternalMaterialID);
 							CellToGeometry.Add(CellIdx, CreatedGeometryIdx);
 							GeometryToResultMesh.Add(CreatedGeometryIdx, CellIdx);
 						}
@@ -1858,7 +1978,12 @@ struct FDynamicMeshCollection
 		return bWasSplit;
 	}
 
-	static int32 AppendToCollection(const FTransform& ToCollection, FDynamicMesh3& Mesh, int32 TransformParent, int32 SubPartIndex, FGeometryCollection& Output, int32 InternalMaterialID)
+	FString GetBoneName(FGeometryCollection& Output, int TransformParent, int SubPartIndex)
+	{
+		return Output.BoneName[TransformParent] + "_" + FString::FromInt(SubPartIndex);
+	}
+
+	static int32 AppendToCollection(const FTransform& ToCollection, FDynamicMesh3& Mesh, double CollisionSampleSpacing, int32 TransformParent, FString BoneName, FGeometryCollection& Output, int32 InternalMaterialID)
 	{
 		if (Mesh.TriangleCount() == 0)
 		{
@@ -1870,6 +1995,11 @@ struct FDynamicMeshCollection
 			Mesh.CompactInPlace(nullptr);
 		}
 
+		if (CollisionSampleSpacing > 0)
+		{
+			UE::PlanarCutInternals::AugmentDynamicMesh::AddCollisionSamplesPerComponent(Mesh, CollisionSampleSpacing);
+		}
+
 		int32 NewGeometryStartIdx = Output.FaceStart.Num();
 		int32 OriginalVertexNum = Output.Vertex.Num();
 		int32 OriginalFaceNum = Output.Indices.Num();
@@ -1889,8 +2019,7 @@ struct FDynamicMeshCollection
 		Output.TransformToGeometryIndex[TransformIdx] = GeometryIdx;
 		if (TransformParent > -1)
 		{
-			// TODO: this is probably not the best way to build the bone name string?
-			Output.BoneName[TransformIdx] = Output.BoneName[TransformParent] + "_" + FString::FromInt(SubPartIndex);
+			Output.BoneName[TransformIdx] = BoneName;
 			Output.BoneColor[TransformIdx] = Output.BoneColor[TransformParent];
 			Output.Parent[TransformIdx] = TransformParent;
 			Output.Children[TransformParent].Add(TransformIdx);
@@ -4001,6 +4130,7 @@ int32 CutWithPlanarCells(
 	FGeometryCollection& Source,
 	int32 TransformIdx,
 	double Grout,
+	double CollisionSampleSpacing,
 	const TOptional<FTransform>& TransformCells,
 	bool bIncludeOutsideCellInOutput,
 	float CheckDistanceAcrossOutsideCellForProximity,
@@ -4009,7 +4139,7 @@ int32 CutWithPlanarCells(
 )
 {
 	TArray<int32> TransformIndices { TransformIdx };
-	return CutMultipleWithPlanarCells(Cells, Source, TransformIndices, Grout, TransformCells, bIncludeOutsideCellInOutput, CheckDistanceAcrossOutsideCellForProximity, bSetDefaultInternalMaterialsFromCollection, VertexInterpolate);
+	return CutMultipleWithPlanarCells(Cells, Source, TransformIndices, Grout, CollisionSampleSpacing, TransformCells, bIncludeOutsideCellInOutput, CheckDistanceAcrossOutsideCellForProximity, bSetDefaultInternalMaterialsFromCollection, VertexInterpolate);
 }
 
 int32 CutMultipleWithPlanarCells_MeshBooleanPath(
@@ -4017,6 +4147,7 @@ int32 CutMultipleWithPlanarCells_MeshBooleanPath(
 	FGeometryCollection& Source,
 	const TArrayView<const int32>& TransformIndices,
 	double Grout,
+	double CollisionSampleSpacing,
 	const TOptional<FTransform>& TransformCells,
 	bool bIncludeOutsideCellInOutput,
 	float CheckDistanceAcrossOutsideCellForProximity,
@@ -4030,6 +4161,7 @@ int32 CutMultipleWithPlanarCells(
 	FGeometryCollection& Source,
 	const TArrayView<const int32>& TransformIndices,
 	double Grout,
+	double CollisionSampleSpacing,
 	const TOptional<FTransform>& TransformCells,
 	bool bIncludeOutsideCellInOutput,
 	float CheckDistanceAcrossOutsideCellForProximity,
@@ -4044,6 +4176,7 @@ int32 CutMultipleWithPlanarCells(
 			Source,
 			TransformIndices,
 			Grout,
+			CollisionSampleSpacing,
 			TransformCells,
 			bIncludeOutsideCellInOutput,
 			CheckDistanceAcrossOutsideCellForProximity,
@@ -4449,6 +4582,7 @@ int32 CutMultipleWithMultiplePlanes(
 	FGeometryCollection& Collection,
 	const TArrayView<const int32>& TransformIndices,
 	double Grout,
+	double CollisionSampleSpacing,
 	const TOptional<FTransform>& TransformCells,
 	bool bSetDefaultInternalMaterialsFromCollection,
 	TFunction<void(const FGeometryCollection&, int32, const FGeometryCollection&, int32, float, int32, FGeometryCollection&)> VertexInterpolate
@@ -4479,7 +4613,7 @@ int32 CutMultipleWithMultiplePlanes(
 	double OnePercentExtend = MeshCollection.Bounds.MaxDim() * .01;
 
 	int32 NewGeomStartIdx = -1;
-	NewGeomStartIdx = MeshCollection.CutWithMultiplePlanes(Planes, Grout, &Collection, InternalSurfaceMaterials, bSetDefaultInternalMaterialsFromCollection);
+	NewGeomStartIdx = MeshCollection.CutWithMultiplePlanes(Planes, Grout, CollisionSampleSpacing, &Collection, InternalSurfaceMaterials, bSetDefaultInternalMaterialsFromCollection);
 
 	Collection.ReindexMaterials();
 	return NewGeomStartIdx;
@@ -4492,6 +4626,7 @@ int32 CutMultipleWithPlanarCells_MeshBooleanPath(
 	FGeometryCollection& Source,
 	const TArrayView<const int32>& TransformIndices,
 	double Grout,
+	double CollisionSampleSpacing,
 	const TOptional<FTransform>& TransformCells,
 	bool bIncludeOutsideCellInOutput,
 	float CheckDistanceAcrossOutsideCellForProximity,
@@ -4518,7 +4653,7 @@ int32 CutMultipleWithPlanarCells_MeshBooleanPath(
 
 	int32 NewGeomStartIdx = -1;
 
-	NewGeomStartIdx = MeshCollection.CutWithCellMeshes(Cells.InternalSurfaceMaterials, Cells.PlaneCells, CellMeshes, &Source, bSetDefaultInternalMaterialsFromCollection);
+	NewGeomStartIdx = MeshCollection.CutWithCellMeshes(Cells.InternalSurfaceMaterials, Cells.PlaneCells, CellMeshes, &Source, bSetDefaultInternalMaterialsFromCollection, CollisionSampleSpacing);
 
 	Source.ReindexMaterials();
 	return NewGeomStartIdx;
@@ -4526,10 +4661,40 @@ int32 CutMultipleWithPlanarCells_MeshBooleanPath(
 
 
 
+int32 AddCollisionSampleVertices(double CollisionSampleSpacing, FGeometryCollection& Collection, const TArrayView<const int32>& TransformIndices)
+{
+	FTransform CellsToWorld = FTransform::Identity;
 
+	FDynamicMeshCollection MeshCollection(&Collection, TransformIndices, CellsToWorld);
+	struct FLocalInfo
+	{
+		FString BoneName;
+		int Parent = -1;
+	};
+	TArray<FLocalInfo> LocalInfo; LocalInfo.SetNum(TransformIndices.Num());
+	for (int Idx = 0; Idx < TransformIndices.Num(); Idx++)
+	{
+		LocalInfo[Idx].BoneName = Collection.BoneName[TransformIndices[Idx]];
+		LocalInfo[Idx].Parent = Collection.Parent[TransformIndices[Idx]];
+	}
 
+	TArray<int32> ToDelete(TransformIndices);
+	ToDelete.Sort();
+	Collection.RemoveElements(FGeometryCollection::TransformGroup, ToDelete);
+	int32 FirstAppended = -1;
+	for (int MeshIdx = 0; MeshIdx < MeshCollection.Meshes.Num(); MeshIdx++)
+	{
+		FDynamicMeshCollection::FMeshData& MeshData = MeshCollection.Meshes[MeshIdx];
+		FDynamicMesh3& Mesh = MeshData.AugMesh;
+		int32 ResultIdx = MeshCollection.AppendToCollection(MeshData.ToCollection, Mesh, CollisionSampleSpacing, LocalInfo[MeshIdx].Parent, LocalInfo[MeshIdx].BoneName, Collection, -1);
+		if (FirstAppended == -1)
+		{
+			FirstAppended = ResultIdx;
+		}
+	}
 
+	Collection.ReindexMaterials();
 
-
-
+	return FirstAppended;
+}
 
diff --git a/Engine/Plugins/Experimental/PlanarCutPlugin/Source/PlanarCut/Public/PlanarCut.h b/Engine/Plugins/Experimental/PlanarCutPlugin/Source/PlanarCut/Public/PlanarCut.h
--- a/Engine/Plugins/Experimental/PlanarCutPlugin/Source/PlanarCut/Public/PlanarCut.h
+++ b/Engine/Plugins/Experimental/PlanarCutPlugin/Source/PlanarCut/Public/PlanarCut.h
@@ -156,6 +156,8 @@ void PLANARCUT_API ComputeTriangleNormals(const TArrayView<const FVector> Vertic
  * @param Cells				Defines the cutting planes and division of space
  * @param Collection		The collection to be cut
  * @param TransformIdx		Which transform inside the collection to cut
+ * @param Grout				Separation to leave between cutting cells
+ * @param CollisionSampleSpacing	Target spacing collision samples	
  * @param TransformCells	Optional transform of the planar cut; if unset, defaults to Identity
  * @param bIncludeOutsideCellInOutput	If true, geometry that was not inside any of the cells (e.g. was outside of the bounds of all cutting geometry) will still be included in the output; if false, it will be discarded.
  * @param CheckDistanceAcrossOutsideCellForProximity	If > 0, when a plane is neighboring the "outside" cell, instead of setting proximity to the outside cell, the algo will sample a point this far outside the cell in the normal direction of the plane to see if there is actually a non-outside cell there.  (Useful for bricks w/out mortar)
@@ -167,6 +169,7 @@ int32 PLANARCUT_API CutWithPlanarCells(
 	FGeometryCollection& Collection,
 	int32 TransformIdx,
 	double Grout,
+	double CollisionSampleSpacing,
 	const TOptional<FTransform>& TransformCells = TOptional<FTransform>(),
 	bool bIncludeOutsideCellInOutput = true,
 	float CheckDistanceAcrossOutsideCellForProximity = 0,
@@ -191,6 +194,7 @@ int32 PLANARCUT_API CutMultipleWithPlanarCells(
 	FGeometryCollection& Collection,
 	const TArrayView<const int32>& TransformIndices,
 	double Grout,
+	double CollisionSampleSpacing,
 	const TOptional<FTransform>& TransformCells = TOptional<FTransform>(),
 	bool bIncludeOutsideCellInOutput = true,
 	float CheckDistanceAcrossOutsideCellForProximity = 0,  // TODO: < this param does nothing in the new mode; is only needed in special cases that aren't possible in the UI currently
@@ -215,7 +219,19 @@ int32 PLANARCUT_API CutMultipleWithMultiplePlanes(
 	FGeometryCollection& Collection,
 	const TArrayView<const int32>& TransformIndices,
 	double Grout,
+	double CollisionSampleSpacing,
 	const TOptional<FTransform>& TransformCells = TOptional<FTransform>(),
 	bool bSetDefaultInternalMaterialsFromCollection = true,
 	TFunction<void(const FGeometryCollection&, int32, const FGeometryCollection&, int32, float, int32, FGeometryCollection&)> VertexInterpolate = DefaultVertexInterpolation
 );
+
+
+/**
+ * Scatter additional vertices (w/ no associated triangle) as needed to satisfy minimum point spacing
+ * 
+ * @param TargetSpacing		The desired spacing between collision sample vertices
+ * @param Collection		The Geometry Collection to be updated
+ * @param TransformIndices	Which transform groups on the Geometry Collection to be updated.  If empty, all groups are updated.
+ * @return Index of first transform group w/ updated geometry.  (To update geometry we delete and re-add, because geometry collection isn't designed for in-place updates)
+ */
+int32 PLANARCUT_API AddCollisionSampleVertices(double TargetSpacing, FGeometryCollection& Collection, const TArrayView<const int32>& TransformIndices = TArrayView<const int32>());
