clean up transforms for fracture tools to make sure the cuts match the previews, including when the geometry collection is transformed
 - changes fracture methods to take the geometry collection transform
 - consistently use world bounding box when generating fracture pattern (previously was used for visualization but not actual cuts)
 - slight refactor of brick tool to update preview more consistently

#rnx

diff --git a/Engine/Plugins/Experimental/ChaosEditor/Source/FractureEditor/Private/FractureToolBrick.cpp b/Engine/Plugins/Experimental/ChaosEditor/Source/FractureEditor/Private/FractureToolBrick.cpp
--- a/Engine/Plugins/Experimental/ChaosEditor/Source/FractureEditor/Private/FractureToolBrick.cpp
+++ b/Engine/Plugins/Experimental/ChaosEditor/Source/FractureEditor/Private/FractureToolBrick.cpp
@@ -195,9 +195,8 @@ void UFractureToolBrick::GenerateBrickTransforms(const FBox& Bounds)
 	}
 }
 
-void UFractureToolBrick::PostEditChangeChainProperty(struct FPropertyChangedChainEvent& PropertyChangedEvent)
+void UFractureToolBrick::UpdateBrickTransforms()
 {
-	
 	FBox Bounds(ForceInitToZero);
 
 	USelection* SelectionSet = GEditor->GetSelectedActors();
@@ -237,10 +236,17 @@ void UFractureToolBrick::PostEditChangeChainProperty(struct FPropertyChangedChai
 	{
 		GenerateBrickTransforms(Bounds);
 	}
-
 }
 
+void UFractureToolBrick::PostEditChangeChainProperty(struct FPropertyChangedChainEvent& PropertyChangedEvent)
+{
+	UpdateBrickTransforms();
+}
 
+void UFractureToolBrick::FractureContextChanged()
+{
+	UpdateBrickTransforms();
+}
 
 void UFractureToolBrick::Render(const FSceneView* View, FViewport* Viewport, FPrimitiveDrawInterface* PDI)
 {
@@ -286,7 +292,7 @@ int32 UFractureToolBrick::ExecuteFracture(const FFractureToolContext& FractureCo
 	if (FractureContext.IsValid())
 	{
  		BrickTransforms.Empty();
-		GenerateBrickTransforms(FractureContext.GetBounds());
+		GenerateBrickTransforms(FractureContext.GetWorldBounds());
 
 		const FQuat HeaderRotation(FVector::UpVector, 1.5708);
 
@@ -320,7 +326,7 @@ int32 UFractureToolBrick::ExecuteFracture(const FFractureToolContext& FractureCo
 			VoronoiPlanarCells.InternalSurfaceMaterials.NoiseSettings = NoiseSettings;
 		}
 
-		return CutMultipleWithPlanarCells(VoronoiPlanarCells, *FractureContext.GetGeometryCollection(), FractureContext.GetSelection(), 0, CollisionSettings->PointSpacing);
+		return CutMultipleWithPlanarCells(VoronoiPlanarCells, *FractureContext.GetGeometryCollection(), FractureContext.GetSelection(), 0, CollisionSettings->PointSpacing, FractureContext.GetTransform());
 	}
 
 	return INDEX_NONE;
diff --git a/Engine/Plugins/Experimental/ChaosEditor/Source/FractureEditor/Private/FractureToolBrick.h b/Engine/Plugins/Experimental/ChaosEditor/Source/FractureEditor/Private/FractureToolBrick.h
--- a/Engine/Plugins/Experimental/ChaosEditor/Source/FractureEditor/Private/FractureToolBrick.h
+++ b/Engine/Plugins/Experimental/ChaosEditor/Source/FractureEditor/Private/FractureToolBrick.h
@@ -93,6 +93,7 @@ public:
 	virtual void RegisterUICommand( FFractureEditorCommands* BindingContext ) override;
 	virtual TArray<UObject*> GetSettingsObjects() const override;
 	void PostEditChangeChainProperty(struct FPropertyChangedChainEvent& PropertyChangedEvent) override;
+	virtual void FractureContextChanged() override;
 	virtual void Render(const FSceneView* View, FViewport* Viewport, FPrimitiveDrawInterface* PDI) override;
 	virtual int32 ExecuteFracture(const FFractureToolContext& FractureContext) override;
 
@@ -105,6 +106,8 @@ private:
 
 	void AddBoxEdges(const FVector& Min, const FVector& Max);
 
+	void UpdateBrickTransforms();
+
 	TArray<FTransform> BrickTransforms;
 	TArray<TTuple<FVector, FVector>> Edges;
 	TArray<TTuple<FVector, FVector>> Boxes;
diff --git a/Engine/Plugins/Experimental/ChaosEditor/Source/FractureEditor/Private/FractureToolClusterCutter.cpp b/Engine/Plugins/Experimental/ChaosEditor/Source/FractureEditor/Private/FractureToolClusterCutter.cpp
--- a/Engine/Plugins/Experimental/ChaosEditor/Source/FractureEditor/Private/FractureToolClusterCutter.cpp
+++ b/Engine/Plugins/Experimental/ChaosEditor/Source/FractureEditor/Private/FractureToolClusterCutter.cpp
@@ -50,14 +50,15 @@ void UFractureToolClusterCutter::GenerateVoronoiSites(const FFractureToolContext
  	FRandomStream RandStream(Context.GetSeed());
 	const int32 SiteCount = RandStream.RandRange(ClusterSettings->NumberClustersMin, ClusterSettings->NumberClustersMax);
 
-	const FVector Extent(Context.GetBounds().Max - Context.GetBounds().Min);
+	FBox Bounds = Context.GetWorldBounds();
+	const FVector Extent(Bounds.Max - Bounds.Min);
 
 	TArray<FVector> CenterSites;
 
 	Sites.Reserve(Sites.Num() + SiteCount);
 	for (int32 ii = 0; ii < SiteCount; ++ii)
 	{
-		CenterSites.Emplace(Context.GetBounds().Min + FVector(RandStream.FRand(), RandStream.FRand(), RandStream.FRand()) * Extent);
+		CenterSites.Emplace(Bounds.Min + FVector(RandStream.FRand(), RandStream.FRand(), RandStream.FRand()) * Extent);
 	}
 
 	for (int32 kk = 0, nk = CenterSites.Num(); kk < nk; ++kk)
@@ -68,7 +69,7 @@ void UFractureToolClusterCutter::GenerateVoronoiSites(const FFractureToolContext
 		{
 			FVector V(RandStream.VRand());
 			V.Normalize();
-			V *= ClusterSettings->ClusterRadius + (RandStream.FRandRange(ClusterSettings->ClusterRadiusPercentageMin, ClusterSettings->ClusterRadiusPercentageMax) * Context.GetBounds().GetExtent().GetAbsMax());
+			V *= ClusterSettings->ClusterRadius + (RandStream.FRandRange(ClusterSettings->ClusterRadiusPercentageMin, ClusterSettings->ClusterRadiusPercentageMax) * Bounds.GetExtent().GetAbsMax());
 			V += CenterSites[kk];
 			Sites.Emplace(V);
 		}
diff --git a/Engine/Plugins/Experimental/ChaosEditor/Source/FractureEditor/Private/FractureToolContext.cpp b/Engine/Plugins/Experimental/ChaosEditor/Source/FractureEditor/Private/FractureToolContext.cpp
--- a/Engine/Plugins/Experimental/ChaosEditor/Source/FractureEditor/Private/FractureToolContext.cpp
+++ b/Engine/Plugins/Experimental/ChaosEditor/Source/FractureEditor/Private/FractureToolContext.cpp
@@ -170,11 +170,6 @@ void FFractureToolContext::ConvertSelectionToClusterNodes()
 	Sanitize();
 }
 
-void FFractureToolContext::TransformBoundsToWorld()
-{
-	Bounds = Bounds.TransformBy(Transform);
-}
-
 
 bool FFractureToolContext::HasSelectedAncestor(int32 Index) const
 {
diff --git a/Engine/Plugins/Experimental/ChaosEditor/Source/FractureEditor/Private/FractureToolContext.h b/Engine/Plugins/Experimental/ChaosEditor/Source/FractureEditor/Private/FractureToolContext.h
--- a/Engine/Plugins/Experimental/ChaosEditor/Source/FractureEditor/Private/FractureToolContext.h
+++ b/Engine/Plugins/Experimental/ChaosEditor/Source/FractureEditor/Private/FractureToolContext.h
@@ -35,8 +35,7 @@ public:
 	/** Return selection in sibling groups */
 	TMap<int32, TArray<int32>> GetClusteredSelections();
 
-	void TransformBoundsToWorld();
-
+	FBox GetWorldBounds() const { return Bounds.TransformBy(Transform); }
 	FBox GetBounds() const { return Bounds; }
 	void SetBounds(FBox& InBounds) { Bounds = InBounds; }
 	int32 GetSeed() const { return RandomSeed; }
diff --git a/Engine/Plugins/Experimental/ChaosEditor/Source/FractureEditor/Private/FractureToolCutter.cpp b/Engine/Plugins/Experimental/ChaosEditor/Source/FractureEditor/Private/FractureToolCutter.cpp
--- a/Engine/Plugins/Experimental/ChaosEditor/Source/FractureEditor/Private/FractureToolCutter.cpp
+++ b/Engine/Plugins/Experimental/ChaosEditor/Source/FractureEditor/Private/FractureToolCutter.cpp
@@ -145,11 +145,10 @@ void UFractureToolVoronoiCutterBase::FractureContextChanged()
 	for (FFractureToolContext& FractureContext : FractureContexts)
 	{
 		// Move the local bounds to the actor so we we'll draw in the correct location
-		FractureContext.TransformBoundsToWorld();
 		GenerateVoronoiSites(FractureContext, VoronoiSites);
 		if (CutterSettings->bDrawDiagram)
 		{
-			GetVoronoiEdges(VoronoiSites, FractureContext.GetBounds(), VoronoiEdges, CellMember);
+			GetVoronoiEdges(VoronoiSites, FractureContext.GetWorldBounds(), VoronoiEdges, CellMember);
 		}
 	}
 }
@@ -160,7 +159,7 @@ int32 UFractureToolVoronoiCutterBase::ExecuteFracture(const FFractureToolContext
 	{
 		TArray<FVector> Sites;
 		GenerateVoronoiSites(FractureContext, Sites);
-		FBox VoronoiBounds = FractureContext.GetBounds();
+		FBox VoronoiBounds = FractureContext.GetWorldBounds();
 			
 		// expand bounds to make sure noise-perturbed voronoi cells still contain the whole input mesh
 		VoronoiBounds = VoronoiBounds.ExpandBy(CutterSettings->Amplitude + CutterSettings->Grout);
@@ -179,7 +178,7 @@ int32 UFractureToolVoronoiCutterBase::ExecuteFracture(const FFractureToolContext
 			VoronoiPlanarCells.InternalSurfaceMaterials.NoiseSettings = NoiseSettings;
 		}
 
-		return CutMultipleWithPlanarCells(VoronoiPlanarCells, *(FractureContext.GetGeometryCollection()), FractureContext.GetSelection(), CutterSettings->Grout, CollisionSettings->PointSpacing);
+		return CutMultipleWithPlanarCells(VoronoiPlanarCells, *(FractureContext.GetGeometryCollection()), FractureContext.GetSelection(), CutterSettings->Grout, CollisionSettings->PointSpacing, FractureContext.GetTransform());
 	}
 
 	return INDEX_NONE;
diff --git a/Engine/Plugins/Experimental/ChaosEditor/Source/FractureEditor/Private/FractureToolPlaneCut.cpp b/Engine/Plugins/Experimental/ChaosEditor/Source/FractureEditor/Private/FractureToolPlaneCut.cpp
--- a/Engine/Plugins/Experimental/ChaosEditor/Source/FractureEditor/Private/FractureToolPlaneCut.cpp
+++ b/Engine/Plugins/Experimental/ChaosEditor/Source/FractureEditor/Private/FractureToolPlaneCut.cpp
@@ -88,12 +88,12 @@ void UFractureToolPlaneCut::FractureContextChanged()
 	for (FFractureToolContext& FractureContext : FractureContexts)
 	{
 		// Move the local bounds to the actor so we we'll draw in the correct location
-		FractureContext.TransformBoundsToWorld();
+		FBox Bounds = FractureContext.GetWorldBounds();
 		GenerateSliceTransforms(FractureContext, RenderCuttingPlanesTransforms);
 
-		if (FractureContext.GetBounds().GetExtent().GetMax() < RenderCuttingPlaneSize)
+		if (Bounds.GetExtent().GetMax() < RenderCuttingPlaneSize)
 		{
-			RenderCuttingPlaneSize = FractureContext.GetBounds().GetExtent().GetMax();
+			RenderCuttingPlaneSize = Bounds.GetExtent().GetMax();
 		}
 	}
 }
@@ -108,7 +108,7 @@ int32 UFractureToolPlaneCut::ExecuteFracture(const FFractureToolContext& Fractur
 		if (LocalCutSettings->ReferenceActor != nullptr)
 		{
 			FTransform Transform(LocalCutSettings->ReferenceActor->GetActorTransform());
-			CuttingPlanes.Add(FPlane(Transform.GetLocation() - FractureContext.GetTransform().GetLocation(), Transform.GetUnitAxis(EAxis::Z)));
+			CuttingPlanes.Add(FPlane(Transform.GetLocation(), Transform.GetUnitAxis(EAxis::Z)));
 		}
 		else
 		{
@@ -131,7 +131,7 @@ int32 UFractureToolPlaneCut::ExecuteFracture(const FFractureToolContext& Fractur
 			InternalSurfaceMaterials.NoiseSettings = NoiseSettings;
 		}
 
-		return CutMultipleWithMultiplePlanes(CuttingPlanes, InternalSurfaceMaterials, *FractureContext.GetGeometryCollection(), FractureContext.GetSelection(), CutterSettings->Grout, CollisionSettings->PointSpacing);
+		return CutMultipleWithMultiplePlanes(CuttingPlanes, InternalSurfaceMaterials, *FractureContext.GetGeometryCollection(), FractureContext.GetSelection(), CutterSettings->Grout, CollisionSettings->PointSpacing, FractureContext.GetTransform());
 	}
 
 	return INDEX_NONE;
@@ -141,12 +141,13 @@ void UFractureToolPlaneCut::GenerateSliceTransforms(const FFractureToolContext&
 {
 	FRandomStream RandStream(Context.GetSeed());
 
-	const FVector Extent(Context.GetBounds().Max - Context.GetBounds().Min);
+	FBox Bounds = Context.GetWorldBounds();
+	const FVector Extent(Bounds.Max - Bounds.Min);
 
 	CuttingPlaneTransforms.Reserve(CuttingPlaneTransforms.Num() + PlaneCutSettings->NumberPlanarCuts);
 	for (int32 ii = 0; ii < PlaneCutSettings->NumberPlanarCuts; ++ii)
 	{
-		FVector Position(Context.GetBounds().Min + FVector(RandStream.FRand(), RandStream.FRand(), RandStream.FRand()) * Extent);
+		FVector Position(Bounds.Min + FVector(RandStream.FRand(), RandStream.FRand(), RandStream.FRand()) * Extent);
 		CuttingPlaneTransforms.Emplace(FTransform(FRotator(RandStream.FRand() * 360.0f, RandStream.FRand() * 360.0f, 0.0f), Position));
 	}
 }
diff --git a/Engine/Plugins/Experimental/ChaosEditor/Source/FractureEditor/Private/FractureToolRadial.cpp b/Engine/Plugins/Experimental/ChaosEditor/Source/FractureEditor/Private/FractureToolRadial.cpp
--- a/Engine/Plugins/Experimental/ChaosEditor/Source/FractureEditor/Private/FractureToolRadial.cpp
+++ b/Engine/Plugins/Experimental/ChaosEditor/Source/FractureEditor/Private/FractureToolRadial.cpp
@@ -49,7 +49,8 @@ void UFractureToolRadial::GenerateVoronoiSites(const FFractureToolContext& Conte
 {
  	float RadialStep = RadialSettings->Radius / RadialSettings->RadialSteps;
 
-	const FVector Center(Context.GetBounds().GetCenter() + RadialSettings->Center);
+	FBox Bounds = Context.GetWorldBounds();
+	const FVector Center(Bounds.GetCenter() + RadialSettings->Center);
 
 	FRandomStream RandStream(Context.GetSeed());
 	FVector UpVector(RadialSettings->Normal);
diff --git a/Engine/Plugins/Experimental/ChaosEditor/Source/FractureEditor/Private/FractureToolSlice.cpp b/Engine/Plugins/Experimental/ChaosEditor/Source/FractureEditor/Private/FractureToolSlice.cpp
--- a/Engine/Plugins/Experimental/ChaosEditor/Source/FractureEditor/Private/FractureToolSlice.cpp
+++ b/Engine/Plugins/Experimental/ChaosEditor/Source/FractureEditor/Private/FractureToolSlice.cpp
@@ -134,12 +134,12 @@ void UFractureToolSlice::FractureContextChanged()
 	for (FFractureToolContext& FractureContext : FractureContexts)
 	{
 		// Move the local bounds to the actor so we we'll draw in the correct location
-		FractureContext.TransformBoundsToWorld();
+		FBox Bounds = FractureContext.GetWorldBounds();
 		GenerateSliceTransforms(FractureContext, RenderCuttingPlanesTransforms);
 
-		if (FractureContext.GetBounds().GetExtent().GetMax() < RenderCuttingPlaneSize)
+		if (Bounds.GetExtent().GetMax() < RenderCuttingPlaneSize)
 		{
-			RenderCuttingPlaneSize = FractureContext.GetBounds().GetExtent().GetMax();
+			RenderCuttingPlaneSize = Bounds.GetExtent().GetMax();
 		}
 	}
 }
diff --git a/Engine/Plugins/Experimental/ChaosEditor/Source/FractureEditor/Private/FractureToolUniform.cpp b/Engine/Plugins/Experimental/ChaosEditor/Source/FractureEditor/Private/FractureToolUniform.cpp
--- a/Engine/Plugins/Experimental/ChaosEditor/Source/FractureEditor/Private/FractureToolUniform.cpp
+++ b/Engine/Plugins/Experimental/ChaosEditor/Source/FractureEditor/Private/FractureToolUniform.cpp
@@ -49,14 +49,15 @@ void UFractureToolUniform::GenerateVoronoiSites(const FFractureToolContext& Cont
 {
 	FRandomStream RandStream(Context.GetSeed());
 
-	const FVector Extent(Context.GetBounds().Max - Context.GetBounds().Min);
+	FBox Bounds = Context.GetWorldBounds();
+	const FVector Extent(Bounds.Max - Bounds.Min);
 
 	const int32 SiteCount = RandStream.RandRange(UniformSettings->NumberVoronoiSitesMin, UniformSettings->NumberVoronoiSitesMax);
 
 	Sites.Reserve(Sites.Num() + SiteCount);
 	for (int32 ii = 0; ii < SiteCount; ++ii)
 	{
-		Sites.Emplace(Context.GetBounds().Min + FVector(RandStream.FRand(), RandStream.FRand(), RandStream.FRand()) * Extent );
+		Sites.Emplace(Bounds.Min + FVector(RandStream.FRand(), RandStream.FRand(), RandStream.FRand()) * Extent );
 	}
 }
 
diff --git a/Engine/Plugins/Experimental/PlanarCutPlugin/Source/PlanarCut/Private/PlanarCut.cpp b/Engine/Plugins/Experimental/PlanarCutPlugin/Source/PlanarCut/Private/PlanarCut.cpp
--- a/Engine/Plugins/Experimental/PlanarCutPlugin/Source/PlanarCut/Private/PlanarCut.cpp
+++ b/Engine/Plugins/Experimental/PlanarCutPlugin/Source/PlanarCut/Private/PlanarCut.cpp
@@ -1277,15 +1277,13 @@ struct FDynamicMeshCollection
 	TArray<FMeshData> Meshes;
 	FAxisAlignedBox3d Bounds;
 
-	FDynamicMeshCollection(const FGeometryCollection* Collection, const TArrayView<const int32>& TransformIndices, FTransform TransformCells)
+	FDynamicMeshCollection(const FGeometryCollection* Collection, const TArrayView<const int32>& TransformIndices, FTransform TransformCollection)
 	{
-		Init(Collection, TransformIndices, TransformCells);
+		Init(Collection, TransformIndices, TransformCollection);
 	}
 
-	void Init(const FGeometryCollection* Collection, const TArrayView<const int32>& TransformIndices, FTransform TransformCells)
+	void Init(const FGeometryCollection* Collection, const TArrayView<const int32>& TransformIndices, FTransform TransformCollection)
 	{
-		FTransform WorldToCells = TransformCells.Inverse();
-
 		Meshes.Reset();
 		Bounds = FAxisAlignedBox3d::Empty();
 
@@ -1297,7 +1295,7 @@ struct FDynamicMeshCollection
 				continue;
 			}
 
-			FTransform3d CollectionToLocal = FTransform3d(GeometryCollectionAlgo::GlobalMatrix(Collection->Transform, Collection->Parent, TransformIdx) * WorldToCells);
+			FTransform3d CollectionToLocal = FTransform3d(GeometryCollectionAlgo::GlobalMatrix(Collection->Transform, Collection->Parent, TransformIdx) * TransformCollection);
 
 			FMeshData& MeshData = Meshes.Emplace_GetRef();
 			MeshData.TransformIndex = TransformIdx;
@@ -4131,7 +4129,7 @@ int32 CutWithPlanarCells(
 	int32 TransformIdx,
 	double Grout,
 	double CollisionSampleSpacing,
-	const TOptional<FTransform>& TransformCells,
+	const TOptional<FTransform>& TransformCollection,
 	bool bIncludeOutsideCellInOutput,
 	float CheckDistanceAcrossOutsideCellForProximity,
 	bool bSetDefaultInternalMaterialsFromCollection,
@@ -4139,7 +4137,7 @@ int32 CutWithPlanarCells(
 )
 {
 	TArray<int32> TransformIndices { TransformIdx };
-	return CutMultipleWithPlanarCells(Cells, Source, TransformIndices, Grout, CollisionSampleSpacing, TransformCells, bIncludeOutsideCellInOutput, CheckDistanceAcrossOutsideCellForProximity, bSetDefaultInternalMaterialsFromCollection, VertexInterpolate);
+	return CutMultipleWithPlanarCells(Cells, Source, TransformIndices, Grout, CollisionSampleSpacing, TransformCollection, bIncludeOutsideCellInOutput, CheckDistanceAcrossOutsideCellForProximity, bSetDefaultInternalMaterialsFromCollection, VertexInterpolate);
 }
 
 int32 CutMultipleWithPlanarCells_MeshBooleanPath(
@@ -4148,7 +4146,7 @@ int32 CutMultipleWithPlanarCells_MeshBooleanPath(
 	const TArrayView<const int32>& TransformIndices,
 	double Grout,
 	double CollisionSampleSpacing,
-	const TOptional<FTransform>& TransformCells,
+	const TOptional<FTransform>& TransformCollection,
 	bool bIncludeOutsideCellInOutput,
 	float CheckDistanceAcrossOutsideCellForProximity,
 	bool bSetDefaultInternalMaterialsFromCollection,
@@ -4162,7 +4160,7 @@ int32 CutMultipleWithPlanarCells(
 	const TArrayView<const int32>& TransformIndices,
 	double Grout,
 	double CollisionSampleSpacing,
-	const TOptional<FTransform>& TransformCells,
+	const TOptional<FTransform>& TransformCollection,
 	bool bIncludeOutsideCellInOutput,
 	float CheckDistanceAcrossOutsideCellForProximity,
 	bool bSetDefaultInternalMaterialsFromCollection,
@@ -4177,7 +4175,7 @@ int32 CutMultipleWithPlanarCells(
 			TransformIndices,
 			Grout,
 			CollisionSampleSpacing,
-			TransformCells,
+			TransformCollection,
 			bIncludeOutsideCellInOutput,
 			CheckDistanceAcrossOutsideCellForProximity,
 			bSetDefaultInternalMaterialsFromCollection,
@@ -4187,7 +4185,7 @@ int32 CutMultipleWithPlanarCells(
 
 	int32 NewGeomStartIdx = -1;
 
-	FTransform CellsToWorld = TransformCells.Get(FTransform::Identity);
+	FTransform CollectionToWorld = TransformCollection.Get(FTransform::Identity);
 
 	TArray<FPlane> TransformedPlanes;
 	TArray<FVector> TransformedPlaneBoundaries;
@@ -4246,7 +4244,7 @@ int32 CutMultipleWithPlanarCells(
 		TArrayView<const FVector> PlaneBoundaries;
 		FTransform LocalToPlaneSpaceTransform = FTransform::Identity;
 
-		LocalToPlaneSpaceTransform = GeometryCollectionAlgo::GlobalMatrix(Source.Transform, Source.Parent, ParentTransformIndex) * CellsToWorld.Inverse();
+		LocalToPlaneSpaceTransform = GeometryCollectionAlgo::GlobalMatrix(Source.Transform, Source.Parent, ParentTransformIndex) * CollectionToWorld;
 		FTransform PlanesToLocalTransform = LocalToPlaneSpaceTransform.Inverse();
 		TransformPlanes(PlanesToLocalTransform, Cells, TransformedPlanes, TransformedPlaneBoundaries);
 		Planes = TransformedPlanes;
@@ -4288,7 +4286,7 @@ int32 CutMultipleWithMultiplePlanes_old(
 	FInternalSurfaceMaterials& InternalSurfaceMaterials,
 	FGeometryCollection& Collection,
 	const TArrayView<const int32>& TransformIndices,
-	const TOptional<FTransform>& TransformCells,
+	const TOptional<FTransform>& TransformCollection,
 	bool bSetDefaultInternalMaterialsFromCollection,
 	TFunction<void(const FGeometryCollection&, int32, const FGeometryCollection&, int32, float, int32, FGeometryCollection&)> VertexInterpolate
 )
@@ -4299,7 +4297,7 @@ int32 CutMultipleWithMultiplePlanes_old(
 	int32 OrigNumGeom = Collection.FaceCount.Num();
 	int32 CurNumGeom = OrigNumGeom;
 
-	FTransform CellsToWorld = TransformCells.Get(FTransform::Identity);
+	FTransform CollectionToWorld = TransformCollection.Get(FTransform::Identity);
 
 	TArray<int32> NeedsCut;
 
@@ -4392,7 +4390,7 @@ int32 CutMultipleWithMultiplePlanes_old(
 			}
 
 			FPlane TransformedPlane = Plane;
-			FTransform LocalToPlaneSpaceTransform = GeometryCollectionAlgo::GlobalMatrix(Collection.Transform, Collection.Parent, ParentTransformIndex) * CellsToWorld.Inverse();
+			FTransform LocalToPlaneSpaceTransform = GeometryCollectionAlgo::GlobalMatrix(Collection.Transform, Collection.Parent, ParentTransformIndex) * CollectionToWorld;
 			FTransform PlanesToLocalTransform = LocalToPlaneSpaceTransform.Inverse();
 			FMatrix Matrix = PlanesToLocalTransform.ToMatrixWithScale();
 			TransformedPlane = Plane.TransformBy(Matrix);
@@ -4583,14 +4581,14 @@ int32 CutMultipleWithMultiplePlanes(
 	const TArrayView<const int32>& TransformIndices,
 	double Grout,
 	double CollisionSampleSpacing,
-	const TOptional<FTransform>& TransformCells,
+	const TOptional<FTransform>& TransformCollection,
 	bool bSetDefaultInternalMaterialsFromCollection,
 	TFunction<void(const FGeometryCollection&, int32, const FGeometryCollection&, int32, float, int32, FGeometryCollection&)> VertexInterpolate
 )
 {
 	if (CVarFractureMethod.GetValueOnAnyThread() <= 0)
 	{
-		return CutMultipleWithMultiplePlanes_old(Planes, InternalSurfaceMaterials, Collection, TransformIndices, TransformCells, bSetDefaultInternalMaterialsFromCollection, VertexInterpolate);
+		return CutMultipleWithMultiplePlanes_old(Planes, InternalSurfaceMaterials, Collection, TransformIndices, TransformCollection, bSetDefaultInternalMaterialsFromCollection, VertexInterpolate);
 	}
 	int32 OrigNumGeom = Collection.FaceCount.Num();
 	int32 CurNumGeom = OrigNumGeom;
@@ -4606,10 +4604,9 @@ int32 CutMultipleWithMultiplePlanes(
 		Collection.AddAttribute<TSet<int32>>("Proximity", FGeometryCollection::GeometryGroup, GeometryDependency);
 	}
 
-	FTransform CellsToWorld = TransformCells.Get(FTransform::Identity);
-	FTransform WorldToCells = CellsToWorld.Inverse();
+	FTransform CollectionToWorld = TransformCollection.Get(FTransform::Identity);
 
-	FDynamicMeshCollection MeshCollection(&Collection, TransformIndices, CellsToWorld);
+	FDynamicMeshCollection MeshCollection(&Collection, TransformIndices, CollectionToWorld);
 	double OnePercentExtend = MeshCollection.Bounds.MaxDim() * .01;
 
 	int32 NewGeomStartIdx = -1;
@@ -4627,7 +4624,7 @@ int32 CutMultipleWithPlanarCells_MeshBooleanPath(
 	const TArrayView<const int32>& TransformIndices,
 	double Grout,
 	double CollisionSampleSpacing,
-	const TOptional<FTransform>& TransformCells,
+	const TOptional<FTransform>& TransformCollection,
 	bool bIncludeOutsideCellInOutput,
 	float CheckDistanceAcrossOutsideCellForProximity,
 	bool bSetDefaultInternalMaterialsFromCollection,
@@ -4645,9 +4642,9 @@ int32 CutMultipleWithPlanarCells_MeshBooleanPath(
 		Cells.InternalSurfaceMaterials.SetUVScaleFromCollection(Source);
 	}
 
-	FTransform CellsToWorld = TransformCells.Get(FTransform::Identity);
+	FTransform CollectionToWorld = TransformCollection.Get(FTransform::Identity);
 
-	FDynamicMeshCollection MeshCollection(&Source, TransformIndices, CellsToWorld);
+	FDynamicMeshCollection MeshCollection(&Source, TransformIndices, CollectionToWorld);
 	double OnePercentExtend = MeshCollection.Bounds.MaxDim() * .01;
 	FCellMeshes CellMeshes(Cells, MeshCollection.Bounds, Grout, OnePercentExtend, bIncludeOutsideCellInOutput);
 
diff --git a/Engine/Plugins/Experimental/PlanarCutPlugin/Source/PlanarCut/Public/PlanarCut.h b/Engine/Plugins/Experimental/PlanarCutPlugin/Source/PlanarCut/Public/PlanarCut.h
--- a/Engine/Plugins/Experimental/PlanarCutPlugin/Source/PlanarCut/Public/PlanarCut.h
+++ b/Engine/Plugins/Experimental/PlanarCutPlugin/Source/PlanarCut/Public/PlanarCut.h
@@ -158,7 +158,7 @@ void PLANARCUT_API ComputeTriangleNormals(const TArrayView<const FVector> Vertic
  * @param TransformIdx		Which transform inside the collection to cut
  * @param Grout				Separation to leave between cutting cells
  * @param CollisionSampleSpacing	Target spacing collision samples	
- * @param TransformCells	Optional transform of the planar cut; if unset, defaults to Identity
+ * @param TransformCollection		Optional transform of the whole geometry collection; if unset, defaults to Identity
  * @param bIncludeOutsideCellInOutput	If true, geometry that was not inside any of the cells (e.g. was outside of the bounds of all cutting geometry) will still be included in the output; if false, it will be discarded.
  * @param CheckDistanceAcrossOutsideCellForProximity	If > 0, when a plane is neighboring the "outside" cell, instead of setting proximity to the outside cell, the algo will sample a point this far outside the cell in the normal direction of the plane to see if there is actually a non-outside cell there.  (Useful for bricks w/out mortar)
  * @param VertexInterpolate	Function that interpolates vertex properties (UVs, normals, etc); a default that handles all the normal vertex properties is provided, should only need to replace this if you have custom attributes
@@ -170,7 +170,7 @@ int32 PLANARCUT_API CutWithPlanarCells(
 	int32 TransformIdx,
 	double Grout,
 	double CollisionSampleSpacing,
-	const TOptional<FTransform>& TransformCells = TOptional<FTransform>(),
+	const TOptional<FTransform>& TransformCollection = TOptional<FTransform>(),
 	bool bIncludeOutsideCellInOutput = true,
 	float CheckDistanceAcrossOutsideCellForProximity = 0,
 	bool bSetDefaultInternalMaterialsFromCollection = true,
@@ -183,7 +183,9 @@ int32 PLANARCUT_API CutWithPlanarCells(
  * @param Cells				Defines the cutting planes and division of space
  * @param Collection		The collection to be cut
  * @param TransformIndices	Which transform groups inside the collection to cut
- * @param TransformCells	Optional transform of the planar cut; if unset, defaults to Identity
+ * @param Grout				Separation to leave between cutting cells
+ * @param CollisionSampleSpacing	Target spacing collision samples
+ * @param TransformCollection		Optional transform of the whole geometry collection; if unset, defaults to Identity
  * @param bIncludeOutsideCellInOutput	If true, geometry that was not inside any of the cells (e.g. was outside of the bounds of all cutting geometry) will still be included in the output; if false, it will be discarded.
  * @param CheckDistanceAcrossOutsideCellForProximity	If > 0, when a plane is neighboring the "outside" cell, instead of setting proximity to the outside cell, the algo will sample a point this far outside the cell in the normal direction of the plane to see if there is actually a non-outside cell there.  (Useful for bricks w/out mortar)
  * @param VertexInterpolate	Function that interpolates vertex properties (UVs, normals, etc); a default that handles all the normal vertex properties is provided, should only need to replace this if you have custom attributes
@@ -195,7 +197,7 @@ int32 PLANARCUT_API CutMultipleWithPlanarCells(
 	const TArrayView<const int32>& TransformIndices,
 	double Grout,
 	double CollisionSampleSpacing,
-	const TOptional<FTransform>& TransformCells = TOptional<FTransform>(),
+	const TOptional<FTransform>& TransformCollection = TOptional<FTransform>(),
 	bool bIncludeOutsideCellInOutput = true,
 	float CheckDistanceAcrossOutsideCellForProximity = 0,  // TODO: < this param does nothing in the new mode; is only needed in special cases that aren't possible in the UI currently
 	bool bSetDefaultInternalMaterialsFromCollection = true,
@@ -209,7 +211,9 @@ int32 PLANARCUT_API CutMultipleWithPlanarCells(
  * @param InternalSurfaceMaterials	Defines material properties for any added internal surfaces
  * @param Collection			The collection to be cut
  * @param TransformIndices	Which transform groups inside the collection to cut
- * @param TransformCells	Optional transform of the planar cut; if unset, defaults to Identity
+ * @param Grout				Separation to leave between cutting cells
+ * @param CollisionSampleSpacing	Target spacing collision samples
+ * @param TransformCollection		Optional transform of the whole geometry collection; if unset, defaults to Identity
  * @param VertexInterpolate	Function that interpolates vertex properties (UVs, normals, etc); a default that handles all the normal vertex properties is provided, should only need to replace this if you have custom attributes
  * @return	index of first new geometry in the Output GeometryCollection, or -1 if no geometry was added
  */
@@ -220,7 +224,7 @@ int32 PLANARCUT_API CutMultipleWithMultiplePlanes(
 	const TArrayView<const int32>& TransformIndices,
 	double Grout,
 	double CollisionSampleSpacing,
-	const TOptional<FTransform>& TransformCells = TOptional<FTransform>(),
+	const TOptional<FTransform>& TransformCollection = TOptional<FTransform>(),
 	bool bSetDefaultInternalMaterialsFromCollection = true,
 	TFunction<void(const FGeometryCollection&, int32, const FGeometryCollection&, int32, float, int32, FGeometryCollection&)> VertexInterpolate = DefaultVertexInterpolation
 );
