mesh cut for fracture

diff --git a/Engine/Plugins/Experimental/ChaosEditor/Source/FractureEditor/Private/FractureEditorCommands.h b/Engine/Plugins/Experimental/ChaosEditor/Source/FractureEditor/Private/FractureEditorCommands.h
--- a/Engine/Plugins/Experimental/ChaosEditor/Source/FractureEditor/Private/FractureEditorCommands.h
+++ b/Engine/Plugins/Experimental/ChaosEditor/Source/FractureEditor/Private/FractureEditorCommands.h
@@ -52,6 +52,7 @@ class FFractureEditorCommands : public TCommands<FFractureEditorCommands>
 		TSharedPtr< FUICommandInfo > Brick;
 		TSharedPtr< FUICommandInfo > Texture;
 		TSharedPtr< FUICommandInfo > Resample;
+		TSharedPtr< FUICommandInfo > Mesh;
 
 		// Property Commands
 		TSharedPtr< FUICommandInfo > SetInitialDynamicState;
diff --git a/Engine/Plugins/Experimental/ChaosEditor/Source/FractureEditor/Private/FractureEditorModeToolkit.cpp b/Engine/Plugins/Experimental/ChaosEditor/Source/FractureEditor/Private/FractureEditorModeToolkit.cpp
--- a/Engine/Plugins/Experimental/ChaosEditor/Source/FractureEditor/Private/FractureEditorModeToolkit.cpp
+++ b/Engine/Plugins/Experimental/ChaosEditor/Source/FractureEditor/Private/FractureEditorModeToolkit.cpp
@@ -497,6 +497,7 @@ void FFractureEditorModeToolkit::BuildToolPalette(FName PaletteIndex, class FToo
 		ToolbarBuilder.AddToolBarButton(Commands.Clustered);
 		ToolbarBuilder.AddToolBarButton(Commands.Radial);
 		ToolbarBuilder.AddToolBarButton(Commands.Planar);
+		ToolbarBuilder.AddToolBarButton(Commands.Mesh);
 		ToolbarBuilder.AddToolBarButton(Commands.Slice);
 		ToolbarBuilder.AddToolBarButton(Commands.Brick);
 		ToolbarBuilder.AddToolBarButton(Commands.Resample);
diff --git a/Engine/Plugins/Experimental/ChaosEditor/Source/FractureEditor/Private/FractureToolMeshCut.cpp b/Engine/Plugins/Experimental/ChaosEditor/Source/FractureEditor/Private/FractureToolMeshCut.cpp
new file mode 100644
--- /dev/null
+++ b/Engine/Plugins/Experimental/ChaosEditor/Source/FractureEditor/Private/FractureToolMeshCut.cpp
@@ -0,0 +1,97 @@
+// Copyright Epic Games, Inc. All Rights Reserved.
+
+#include "FractureToolMeshCut.h"
+#include "FractureEditorStyle.h"
+#include "PlanarCut.h"
+#include "GeometryCollection/GeometryCollection.h"
+#include "GeometryCollection/GeometryCollectionObject.h"
+#include "FractureEditorModeToolkit.h"
+#include "FractureToolContext.h"
+
+#define LOCTEXT_NAMESPACE "FractureMesh"
+
+
+UFractureToolMeshCut::UFractureToolMeshCut(const FObjectInitializer& ObjInit) 
+	: Super(ObjInit) 
+{
+	MeshCutSettings = NewObject<UFractureMeshCutSettings>(GetTransientPackage(), UFractureMeshCutSettings::StaticClass());
+	MeshCutSettings->OwnerTool = this;
+}
+
+FText UFractureToolMeshCut::GetDisplayText() const
+{
+	return FText(NSLOCTEXT("Fracture", "FractureToolMeshCut", "Mesh Cut Fracture")); 
+}
+
+FText UFractureToolMeshCut::GetTooltipText() const 
+{
+	return FText(NSLOCTEXT("Fracture", "FractureToolMeshCutTooltip", "Mesh fracture can be used to make cuts along a mesh in your Geometry Collection. You can apply noise to mesh cuts for more organic results.  Click the Fracture Button to commit the fracture to the geometry collection."));
+}
+
+FSlateIcon UFractureToolMeshCut::GetToolIcon() const 
+{
+	return FSlateIcon("FractureEditorStyle", "FractureEditor.Mesh");
+}
+
+void UFractureToolMeshCut::RegisterUICommand( FFractureEditorCommands* BindingContext ) 
+{
+	UI_COMMAND_EXT( BindingContext, UICommandInfo, "Mesh", "Mesh", "Mesh Fracture", EUserInterfaceActionType::ToggleButton, FInputChord() );
+	BindingContext->Mesh = UICommandInfo;
+}
+
+TArray<UObject*> UFractureToolMeshCut::GetSettingsObjects() const
+ {
+	TArray<UObject*> Settings;
+	//Settings.Add(CutterSettings); // TODO: add cutter settings if/when we support noise and grout
+	Settings.Add(CollisionSettings);
+	Settings.Add(MeshCutSettings);
+	return Settings;
+}
+
+bool UFractureToolMeshCut::IsCuttingActorValid()
+{
+	const UFractureMeshCutSettings* LocalCutSettings = MeshCutSettings;
+	if (LocalCutSettings->CuttingActor == nullptr)
+	{
+		return false;
+	}
+	const UStaticMeshComponent* Component = LocalCutSettings->CuttingActor->GetStaticMeshComponent();
+	if (Component == nullptr)
+	{
+		return false;
+	}
+	const UStaticMesh* Mesh = Component->GetStaticMesh();
+	if (Mesh == nullptr)
+	{
+		return false;
+	}
+	if (Mesh->GetNumLODs() < 1)
+	{
+		return false;
+	}
+	return true;
+}
+
+int32 UFractureToolMeshCut::ExecuteFracture(const FFractureToolContext& FractureContext)
+{
+	if (FractureContext.IsValid())
+	{
+		const UFractureMeshCutSettings* LocalCutSettings = MeshCutSettings;
+
+		if (!IsCuttingActorValid())
+		{
+			return INDEX_NONE;
+		}
+
+		FMeshDescription* MeshDescription = LocalCutSettings->CuttingActor->GetStaticMeshComponent()->GetStaticMesh()->GetMeshDescription(0);
+		FTransform Transform(LocalCutSettings->CuttingActor->GetTransform());
+		FInternalSurfaceMaterials InternalSurfaceMaterials;
+		// (Note: noise and grout not currently supported)
+		return CutWithMesh(MeshDescription, Transform, InternalSurfaceMaterials, *FractureContext.GetGeometryCollection(), FractureContext.GetSelection(), CollisionSettings->PointSpacing, FractureContext.GetTransform());
+	}
+
+	return INDEX_NONE;
+}
+
+
+#undef LOCTEXT_NAMESPACE
\ No newline at end of file
diff --git a/Engine/Plugins/Experimental/ChaosEditor/Source/FractureEditor/Private/FractureToolMeshCut.h b/Engine/Plugins/Experimental/ChaosEditor/Source/FractureEditor/Private/FractureToolMeshCut.h
new file mode 100644
--- /dev/null
+++ b/Engine/Plugins/Experimental/ChaosEditor/Source/FractureEditor/Private/FractureToolMeshCut.h
@@ -0,0 +1,57 @@
+// Copyright Epic Games, Inc. All Rights Reserved.
+
+#pragma once
+
+#include "FractureToolCutter.h"
+#include "Engine/StaticMeshActor.h"
+
+#include "FractureToolMeshCut.generated.h"
+
+class FFractureToolContext;
+
+
+UCLASS(config = EditorPerProjectUserSettings)
+class UFractureMeshCutSettings : public UFractureToolSettings
+{
+public:
+	GENERATED_BODY()
+
+	UFractureMeshCutSettings(const FObjectInitializer& ObjInit)
+		: Super(ObjInit) {}
+
+	/** Static mesh actor to be used as a cutting surface; should be a closed, watertight mesh */
+	UPROPERTY(EditAnywhere, Category = MeshCut, meta = (DisplayName = "Cutting Actor"))
+	TLazyObjectPtr<AStaticMeshActor> CuttingActor;
+
+};
+
+
+UCLASS(DisplayName = "Mesh Cut Tool", Category = "FractureTools")
+class UFractureToolMeshCut : public UFractureToolCutterBase
+{
+public:
+	GENERATED_BODY()
+
+	UFractureToolMeshCut(const FObjectInitializer& ObjInit);
+
+	// UFractureTool Interface
+	virtual FText GetDisplayText() const override;
+	virtual FText GetTooltipText() const override;
+	virtual FSlateIcon GetToolIcon() const override;
+
+	virtual void RegisterUICommand(FFractureEditorCommands* BindingContext) override;
+
+	virtual TArray<UObject*> GetSettingsObjects() const override;
+
+	virtual int32 ExecuteFracture(const FFractureToolContext& FractureContext) override;
+
+private:
+	// Slicing
+	UPROPERTY(EditAnywhere, Category = Slicing)
+	UFractureMeshCutSettings* MeshCutSettings;
+
+	// check if the chosen actor can be used to cut the geometry collection (i.e. if it is a valid static mesh actor with a non-empty static mesh)
+	bool IsCuttingActorValid();
+};
+
+
diff --git a/Engine/Plugins/Experimental/PlanarCutPlugin/Source/PlanarCut/PlanarCut.Build.cs b/Engine/Plugins/Experimental/PlanarCutPlugin/Source/PlanarCut/PlanarCut.Build.cs
--- a/Engine/Plugins/Experimental/PlanarCutPlugin/Source/PlanarCut/PlanarCut.Build.cs
+++ b/Engine/Plugins/Experimental/PlanarCutPlugin/Source/PlanarCut/PlanarCut.Build.cs
@@ -18,7 +18,9 @@ namespace UnrealBuildTool.Rules
 					"GeometricObjects",
 					"DynamicMesh",
 					"Chaos",
-					"GeometryAlgorithms"
+					"GeometryAlgorithms",
+					"MeshDescription",
+					"MeshConversion",
                 }
                 );
 		}
diff --git a/Engine/Plugins/Experimental/PlanarCutPlugin/Source/PlanarCut/Private/PlanarCut.cpp b/Engine/Plugins/Experimental/PlanarCutPlugin/Source/PlanarCut/Private/PlanarCut.cpp
--- a/Engine/Plugins/Experimental/PlanarCutPlugin/Source/PlanarCut/Private/PlanarCut.cpp
+++ b/Engine/Plugins/Experimental/PlanarCutPlugin/Source/PlanarCut/Private/PlanarCut.cpp
@@ -6,6 +6,7 @@
 #include "Spatial/FastWinding.h"
 #include "Spatial/PointHashGrid3.h"
 #include "Spatial/MeshSpatialSort.h"
+#include "Util/IndexUtil.h"
 #include "Arrangement2d.h"
 #include "MeshAdapter.h"
 #include "FrameTypes.h"
@@ -30,6 +31,8 @@
 #include "MeshNormals.h"
 #include "ConstrainedDelaunay2.h"
 
+#include "MeshDescriptionToDynamicMesh.h"
+
 #include "Algo/Rotate.h"
 
 #if WITH_EDITOR
@@ -361,6 +364,31 @@ struct FCellMeshes
 	{
 		Init(Cells, DomainBounds, Grout, ExtendDomain, bIncludeOutsideCell);
 	}
+
+	FCellMeshes(FDynamicMesh3& SingleCutter, const FInternalSurfaceMaterials& Materials, TOptional<FTransform> Transform)
+	{
+		CellMeshes.Reset();
+		CellMeshes.SetNum(2);
+
+		if (Transform.IsSet())
+		{
+			MeshTransforms::ApplyTransform(SingleCutter, FTransform3d(Transform.GetValue()));
+		}
+
+		// Mesh should already be augmented
+		if (!ensure(UE::PlanarCutInternals::AugmentDynamicMesh::IsAugmented(SingleCutter)))
+		{
+			UE::PlanarCutInternals::AugmentDynamicMesh::Augment(SingleCutter);
+		}
+
+		CellMeshes[0].AugMesh = SingleCutter;
+
+		// first mesh is the same as the second mesh, but will be subtracted b/c it's the "outside cell"
+		// TODO: special case this logic so we don't have to hold two copies of the exact same mesh!
+		CellMeshes[1].AugMesh = CellMeshes[0].AugMesh;
+		OutsideCellIndex = 1;
+
+	}
 	
 	// Special function to just make the "grout" part of the planar mesh cells
 	// Used to make the multi-plane cuts with grout easier to implement
@@ -4658,6 +4686,117 @@ int32 CutMultipleWithPlanarCells_MeshBooleanPath(
 
 
 
+
+int32 CutWithMesh(
+	FMeshDescription* CuttingMesh,
+	FTransform CuttingMeshTransform,
+	FInternalSurfaceMaterials& InternalSurfaceMaterials,
+	FGeometryCollection& Collection,
+	const TArrayView<const int32>& TransformIndices,
+	double CollisionSampleSpacing,
+	const TOptional<FTransform>& TransformCollection,
+	bool bSetDefaultInternalMaterialsFromCollection
+)
+{
+	int32 NewGeomStartIdx = -1;
+
+	// populate the BaseMesh with a conversion of the input mesh.
+	FMeshDescriptionToDynamicMesh Converter;
+	FDynamicMesh3 FullMesh; // full-featured conversion of the source mesh
+	Converter.Convert(CuttingMesh, FullMesh, true);
+	FDynamicMesh3 DynamicCuttingMesh; // version of mesh that is split apart at seams to be compatible w/ geometry collection, with corresponding attributes set
+	UE::PlanarCutInternals::AugmentDynamicMesh::Augment(DynamicCuttingMesh);
+	// Note: This conversion will likely go away, b/c I plan to switch over to doing the boolean operations on the fuller rep, but the code can be adapted
+	//		 to the dynamic mesh -> geometry collection conversion phase, as this same splitting will then need to happen there.
+	if (ensure(FullMesh.HasAttributes() && FullMesh.Attributes()->NumUVLayers() >= 1 && FullMesh.Attributes()->NumNormalLayers() == 3))
+	{
+		if (!ensure(FullMesh.IsCompact()))
+		{
+			FullMesh.CompactInPlace();
+		}
+		// Triangles array is 1:1 with the input mesh
+		TArray<FIndex3i> Triangles; Triangles.Init(FIndex3i::Invalid(), FullMesh.TriangleCount());
+		
+		FDynamicMesh3& OutMesh = DynamicCuttingMesh;
+		FDynamicMeshAttributeSet& Attribs = *FullMesh.Attributes();
+		FDynamicMeshNormalOverlay* NTB[3]{ Attribs.PrimaryNormals(), Attribs.PrimaryTangents(), Attribs.PrimaryBiTangents() };
+		FDynamicMeshUVOverlay* UV = Attribs.PrimaryUV();
+		TMap<FIndex4i, int> ElIDsToVID;
+		int OrigMaxVID = FullMesh.MaxVertexID();
+		for (int VID = 0; VID < OrigMaxVID; VID++)
+		{
+			check(FullMesh.IsVertex(VID));
+			FVector3d Pos = FullMesh.GetVertex(VID);
+
+			ElIDsToVID.Reset();
+			FullMesh.EnumerateVertexTriangles(VID, [&FullMesh, &Triangles, &OutMesh, &NTB, &UV, &ElIDsToVID, Pos, VID](int32 TID)
+			{
+				FIndex3i InTri = FullMesh.GetTriangle(TID);
+				int VOnT = IndexUtil::FindTriIndex(VID, InTri);
+				FIndex4i ElIDs(
+					NTB[0]->GetTriangle(TID)[VOnT],
+					NTB[1]->GetTriangle(TID)[VOnT],
+					NTB[2]->GetTriangle(TID)[VOnT],
+					UV->GetTriangle(TID)[VOnT]);
+				const int* FoundVID = ElIDsToVID.Find(ElIDs);
+
+				FIndex3i& OutTri = Triangles[TID];
+				if (FoundVID)
+				{
+					OutTri[VOnT] = *FoundVID;
+				}
+				else
+				{
+					FVector3f Normal = NTB[0]->GetElement(ElIDs.A);
+					FVertexInfo Info(Pos, Normal, FVector3f(1, 1, 1), UV->GetElement(ElIDs.D));
+
+					int OutVID = OutMesh.AppendVertex(Info);
+					OutTri[VOnT] = OutVID;
+					UE::PlanarCutInternals::AugmentDynamicMesh::SetTangent(OutMesh, OutVID, Normal, NTB[1]->GetElement(ElIDs.B), NTB[2]->GetElement(ElIDs.C));
+					ElIDsToVID.Add(ElIDs, OutVID);
+				}
+			});
+		}
+
+		FDynamicMeshMaterialAttribute* OutMaterialID = OutMesh.Attributes()->GetMaterialID();
+		for (int TID = 0; TID < Triangles.Num(); TID++)
+		{
+			FIndex3i& Tri = Triangles[TID];
+			OutMesh.AppendTriangle(Tri);
+			OutMaterialID->SetValue(TID, -1); // just use a single negative material ID by convention to indicate internal material
+			UE::PlanarCutInternals::AugmentDynamicMesh::SetVisibility(OutMesh, TID, true);
+		}
+	}
+
+	if (!Collection.HasAttribute("Proximity", FGeometryCollection::GeometryGroup))
+	{
+		const FManagedArrayCollection::FConstructionParameters GeometryDependency(FGeometryCollection::GeometryGroup);
+		Collection.AddAttribute<TSet<int32>>("Proximity", FGeometryCollection::GeometryGroup, GeometryDependency);
+	}
+
+	if (bSetDefaultInternalMaterialsFromCollection)
+	{
+		InternalSurfaceMaterials.SetUVScaleFromCollection(Collection);
+	}
+
+	ensureMsgf(!InternalSurfaceMaterials.NoiseSettings.IsSet(), TEXT("Noise settings not yet supported for mesh-based fracture"));
+
+	FTransform CollectionToWorld = TransformCollection.Get(FTransform::Identity);
+
+	FDynamicMeshCollection MeshCollection(&Collection, TransformIndices, CollectionToWorld);
+	FCellMeshes CellMeshes(DynamicCuttingMesh, InternalSurfaceMaterials, CuttingMeshTransform);
+
+	TArray<TPair<int32, int32>> CellConnectivity;
+	CellConnectivity.Add(TPair<int32, int32>(0, -1)); // there's only one 'inside' cell (0), so all cut surfaces are connecting the 'inside' cell (0) to the 'outside' cell (-1)
+
+	NewGeomStartIdx = MeshCollection.CutWithCellMeshes(InternalSurfaceMaterials, CellConnectivity, CellMeshes, &Collection, bSetDefaultInternalMaterialsFromCollection, CollisionSampleSpacing);
+
+	Collection.ReindexMaterials();
+	return NewGeomStartIdx;
+}
+
+
+
 int32 AddCollisionSampleVertices(double CollisionSampleSpacing, FGeometryCollection& Collection, const TArrayView<const int32>& TransformIndices)
 {
 	FTransform CellsToWorld = FTransform::Identity;
diff --git a/Engine/Plugins/Experimental/PlanarCutPlugin/Source/PlanarCut/Public/PlanarCut.h b/Engine/Plugins/Experimental/PlanarCutPlugin/Source/PlanarCut/Public/PlanarCut.h
--- a/Engine/Plugins/Experimental/PlanarCutPlugin/Source/PlanarCut/Public/PlanarCut.h
+++ b/Engine/Plugins/Experimental/PlanarCutPlugin/Source/PlanarCut/Public/PlanarCut.h
@@ -6,6 +6,7 @@
 
 #include "Voronoi/Voronoi.h"
 #include "GeometryCollection/GeometryCollection.h"
+#include "MeshDescription.h"
 
 struct PLANARCUT_API FNoiseSettings
 {
@@ -157,7 +158,7 @@ void PLANARCUT_API ComputeTriangleNormals(const TArrayView<const FVector> Vertic
  * @param Collection		The collection to be cut
  * @param TransformIdx		Which transform inside the collection to cut
  * @param Grout				Separation to leave between cutting cells
- * @param CollisionSampleSpacing	Target spacing collision samples	
+ * @param CollisionSampleSpacing	Target spacing between collision sample vertices	
  * @param TransformCollection		Optional transform of the whole geometry collection; if unset, defaults to Identity
  * @param bIncludeOutsideCellInOutput	If true, geometry that was not inside any of the cells (e.g. was outside of the bounds of all cutting geometry) will still be included in the output; if false, it will be discarded.
  * @param CheckDistanceAcrossOutsideCellForProximity	If > 0, when a plane is neighboring the "outside" cell, instead of setting proximity to the outside cell, the algo will sample a point this far outside the cell in the normal direction of the plane to see if there is actually a non-outside cell there.  (Useful for bricks w/out mortar)
@@ -184,7 +185,7 @@ int32 PLANARCUT_API CutWithPlanarCells(
  * @param Collection		The collection to be cut
  * @param TransformIndices	Which transform groups inside the collection to cut
  * @param Grout				Separation to leave between cutting cells
- * @param CollisionSampleSpacing	Target spacing collision samples
+ * @param CollisionSampleSpacing	Target spacing between collision sample vertices
  * @param TransformCollection		Optional transform of the whole geometry collection; if unset, defaults to Identity
  * @param bIncludeOutsideCellInOutput	If true, geometry that was not inside any of the cells (e.g. was outside of the bounds of all cutting geometry) will still be included in the output; if false, it will be discarded.
  * @param CheckDistanceAcrossOutsideCellForProximity	If > 0, when a plane is neighboring the "outside" cell, instead of setting proximity to the outside cell, the algo will sample a point this far outside the cell in the normal direction of the plane to see if there is actually a non-outside cell there.  (Useful for bricks w/out mortar)
@@ -212,7 +213,7 @@ int32 PLANARCUT_API CutMultipleWithPlanarCells(
  * @param Collection			The collection to be cut
  * @param TransformIndices	Which transform groups inside the collection to cut
  * @param Grout				Separation to leave between cutting cells
- * @param CollisionSampleSpacing	Target spacing collision samples
+ * @param CollisionSampleSpacing	Target spacing between collision sample vertices
  * @param TransformCollection		Optional transform of the whole geometry collection; if unset, defaults to Identity
  * @param VertexInterpolate	Function that interpolates vertex properties (UVs, normals, etc); a default that handles all the normal vertex properties is provided, should only need to replace this if you have custom attributes
  * @return	index of first new geometry in the Output GeometryCollection, or -1 if no geometry was added
@@ -239,3 +240,28 @@ int32 PLANARCUT_API CutMultipleWithMultiplePlanes(
  * @return Index of first transform group w/ updated geometry.  (To update geometry we delete and re-add, because geometry collection isn't designed for in-place updates)
  */
 int32 PLANARCUT_API AddCollisionSampleVertices(double TargetSpacing, FGeometryCollection& Collection, const TArrayView<const int32>& TransformIndices = TArrayView<const int32>());
+
+/**
+ * Cut multiple Geometry groups inside a GeometryCollection with Planes, and add each cut cell back to the GeometryCollection as a new child of their source Geometry.  For geometries that would not be cut, nothing is added.
+ * 
+ * @param CuttingMesh				Mesh to be used to cut the geometry collection
+ * @param CuttingMeshTransform		Position of cutting mesh
+ * @param InternalSurfaceMaterials	Defines material properties for any added internal surfaces
+ * @param Collection				The collection to be cut
+ * @param TransformIndices			Which transform groups inside the collection to cut
+ * @param CollisionSampleSpacing	Target spacing between collision sample vertices
+ * @param TransformCollection		Optional transform of the collection; if unset, defaults to Identity
+ * @return index of first new geometry in the Output GeometryCollection, or -1 if no geometry was added
+ */
+int32 PLANARCUT_API CutWithMesh(
+	FMeshDescription* CuttingMesh,
+	FTransform CuttingMeshTransform,
+	FInternalSurfaceMaterials& InternalSurfaceMaterials,
+	FGeometryCollection& Collection,
+	const TArrayView<const int32>& TransformIndices,
+	double CollisionSampleSpacing,
+	const TOptional<FTransform>& TransformCollection = TOptional<FTransform>(),
+	bool bSetDefaultInternalMaterialsFromCollection = true
+);
+
+
