fix mesh fracture for negative scales

diff --git a/Engine/Plugins/Experimental/PlanarCutPlugin/Source/PlanarCut/Private/PlanarCut.cpp b/Engine/Plugins/Experimental/PlanarCutPlugin/Source/PlanarCut/Private/PlanarCut.cpp
--- a/Engine/Plugins/Experimental/PlanarCutPlugin/Source/PlanarCut/Private/PlanarCut.cpp
+++ b/Engine/Plugins/Experimental/PlanarCutPlugin/Source/PlanarCut/Private/PlanarCut.cpp
@@ -365,29 +365,32 @@ struct FCellMeshes
 		Init(Cells, DomainBounds, Grout, ExtendDomain, bIncludeOutsideCell);
 	}
 
-	FCellMeshes(FDynamicMesh3& SingleCutter, const FInternalSurfaceMaterials& Materials, TOptional<FTransform> Transform)
+	FCellMeshes(const FDynamicMesh3& SingleCutter, const FInternalSurfaceMaterials& Materials, TOptional<FTransform> Transform)
 	{
 		CellMeshes.Reset();
 		CellMeshes.SetNum(2);
 
+		CellMeshes[0].AugMesh = SingleCutter;
+
 		if (Transform.IsSet())
 		{
-			MeshTransforms::ApplyTransform(SingleCutter, FTransform3d(Transform.GetValue()));
+			MeshTransforms::ApplyTransform(CellMeshes[0].AugMesh, FTransform3d(Transform.GetValue()));
+			if (Transform->GetDeterminant() < 0)
+			{
+				CellMeshes[0].AugMesh.ReverseOrientation(false);
+			}
 		}
 
 		// Mesh should already be augmented
-		if (!ensure(UE::PlanarCutInternals::AugmentDynamicMesh::IsAugmented(SingleCutter)))
+		if (!ensure(UE::PlanarCutInternals::AugmentDynamicMesh::IsAugmented(CellMeshes[0].AugMesh)))
 		{
-			UE::PlanarCutInternals::AugmentDynamicMesh::Augment(SingleCutter);
+			UE::PlanarCutInternals::AugmentDynamicMesh::Augment(CellMeshes[0].AugMesh);
 		}
 
-		CellMeshes[0].AugMesh = SingleCutter;
-
 		// first mesh is the same as the second mesh, but will be subtracted b/c it's the "outside cell"
 		// TODO: special case this logic so we don't have to hold two copies of the exact same mesh!
 		CellMeshes[1].AugMesh = CellMeshes[0].AugMesh;
 		OutsideCellIndex = 1;
-
 	}
 	
 	// Special function to just make the "grout" part of the planar mesh cells
