remove a bunch of unused code from PlanarCut module (old fracture algorithm and associated code)

#rnx

diff --git a/Engine/Plugins/Experimental/PlanarCutPlugin/Source/PlanarCut/Private/PlanarCut.cpp b/Engine/Plugins/Experimental/PlanarCutPlugin/Source/PlanarCut/Private/PlanarCut.cpp
--- a/Engine/Plugins/Experimental/PlanarCutPlugin/Source/PlanarCut/Private/PlanarCut.cpp
+++ b/Engine/Plugins/Experimental/PlanarCutPlugin/Source/PlanarCut/Private/PlanarCut.cpp
@@ -39,15 +39,6 @@
 #include "Misc/ScopedSlowTask.h"
 #endif
 
-// Console variable to toggle between new and old fracture method; may delete this when we delete the old method
-static TAutoConsoleVariable<int32> CVarFractureMethod(
-	TEXT("Fracture.MeshCuttingMethod"),
-	1,
-	TEXT("Defines the which method is used under the hood to compute fracture geometry.\n")
-	TEXT("<=0: original, with fake noise\n")
-	TEXT("  1: experimental new method\n"),
-	ECVF_Scalability);
-
 
 /**
  * Adapter that lets the generic fast winding and AABB tree code view a geometry from a geometrycollection as a simple index buffer triangle mesh
@@ -2174,39 +2165,12 @@ void FInternalSurfaceMaterials::SetUVScaleFromCollection(const FGeometryCollecti
 	}
 }
 
-inline double PlaneDotDouble(const FPlane& Plane, const FVector& V)
-{
-	return (double)Plane.X*V.X + Plane.Y*V.Y + Plane.Z*V.Z - Plane.W;
-}
-
-
-inline int PlaneSide(const FPlane &Plane, const FVector& V, const double Epsilon = 1e-3)
-{
-	double SD = PlaneDotDouble(Plane, V);
-	return SD > Epsilon ? 1 : SD < -Epsilon ? -1 : 0;
-}
-
-// TODO: warning -- If Epsilon is too small, we can hit an infinite loop on mesh cutting (if the edge cut is still seen as crossing!)
-inline bool IsSegmentCrossing(const FPlane &Plane, const FVector& A, const FVector& B, double& CrossingT, const double Epsilon = 1e-3)
-{
-	double SDA = PlaneDotDouble(Plane, A);
-	double SDB = PlaneDotDouble(Plane, B);
-	CrossingT = SDA / (SDA - SDB);
-	int32 ASide = SDA < -Epsilon ? -1 : SDA > Epsilon ? 1 : 0;
-	int32 BSide = SDB < -Epsilon ? -1 : SDB > Epsilon ? 1 : 0;
-	return ASide * BSide == -1 && CrossingT < 1 - Epsilon && CrossingT > Epsilon;
-}
 
 
 FPlanarCells::FPlanarCells(const FPlane& P)
 {
 	NumCells = 2;
 	AddPlane(P, 0, 1);
-
-	CellFromPosition = TFunction<int32(FVector)>([P](FVector Position)
-	{
-		return PlaneDotDouble(P, Position) > 0 ? 1 : 0;
-	});
 }
 
 FPlanarCells::FPlanarCells(const TArrayView<const FVector> Sites, FVoronoiDiagram& Voronoi)
@@ -2216,10 +2180,6 @@ FPlanarCells::FPlanarCells(const TArrayView<const FVector> Sites, FVoronoiDiagra
 
 	AssumeConvexCells = true;
 	NumCells = VoronoiCells.Num();
-	CellFromPosition = TFunction<int32(FVector)>([&Voronoi](FVector Position)
-	{
-		return Voronoi.FindCell(Position);
-	});
 	for (int32 CellIdx = 0; CellIdx < NumCells; CellIdx++)
 	{
 		int32 LocalVertexStart = -1;
@@ -2273,18 +2233,7 @@ FPlanarCells::FPlanarCells(const TArrayView<const FBox> Boxes)
 	AssumeConvexCells = true;
 	NumCells = Boxes.Num();
 	TArray<FBox> BoxesCopy(Boxes);
-	CellFromPosition = TFunction<int32(FVector)>([BoxesCopy](FVector Position)
-	{
-		for (int32 Idx = 0; Idx < BoxesCopy.Num(); Idx++)
-		{
-			if (BoxesCopy[Idx].IsInsideOrOn(Position))
-			{
-				return Idx;
-			}
-		}
-		return -1;
-	});
-
+	
 	for (int32 BoxIdx = 0; BoxIdx < NumCells; BoxIdx++)
 	{
 		const FBox &Box = Boxes[BoxIdx];
@@ -2316,25 +2265,6 @@ FPlanarCells::FPlanarCells(const FBox& Region, const FIntVector& CubesPerAxis)
 	AssumeConvexCells = true;
 	NumCells = CubesPerAxis.X * CubesPerAxis.Y * CubesPerAxis.Z;
 
-	CellFromPosition = TFunction<int32(FVector)>([Region, CubesPerAxis](FVector Position)
-	{
-		if (!Region.IsInsideOrOn(Position))
-		{
-			return -1;
-		}
-		FVector Diagonal = Region.Max - Region.Min;
-		FVector RelPos = Position - Region.Min;
-		FIntVector GridIdx(
-			CubesPerAxis.X * (RelPos.X / Diagonal.X), 
-			CubesPerAxis.Y * (RelPos.Y / Diagonal.Y), 
-			CubesPerAxis.Z * (RelPos.Z / Diagonal.Z)
-		);
-		GridIdx.X = FMath::Clamp(GridIdx.X, 0, CubesPerAxis.X - 1);
-		GridIdx.Y = FMath::Clamp(GridIdx.Y, 0, CubesPerAxis.Y - 1);
-		GridIdx.Z = FMath::Clamp(GridIdx.Z, 0, CubesPerAxis.Z - 1);
-		return GridIdx.X + GridIdx.Y * (CubesPerAxis.X) + GridIdx.Z * (CubesPerAxis.X * CubesPerAxis.Y);
-	});
-
 	// cube X, Y, Z integer indices to a single cell index
 	auto ToIdx = [](const FIntVector &PerAxis, int32 Xi, int32 Yi, int32 Zi)
 	{
@@ -2637,2058 +2567,106 @@ FPlanarCells::FPlanarCells(const FBox &Region, const TArrayView<const FColor> Im
 
 
 	AssumeConvexCells = false; // todo could set this to true if the 2D shape of each image region is convex
-
-	CellFromPosition = TFunction<int32(FVector)>([Region, PixCells, Width, Height, RegionDiagonal](FVector Position)
-	{
-		if (!Region.IsInsideOrOn(Position))
-		{
-			return -1;
-		}
-		
-		FVector RelPos = Position - Region.Min;
-		int32 Xg = Width * (RelPos.X / RegionDiagonal.X);
-		int32 Yg = Height * (RelPos.Y / RegionDiagonal.Y);
-		
-		Xg = FMath::Clamp(Xg, 0, Width - 1);
-		Yg = FMath::Clamp(Yg, 0, Height - 1);
-		return PixCells[Xg + Yg * Width];
-	});
 }
 
-void PLANARCUT_API DefaultVertexInterpolation(const FGeometryCollection& V0Collection, int32 V0, const FGeometryCollection& V1Collection, int32 V1, float T, int32 VOut, FGeometryCollection& Dest)
-{
-	// For now just manually write the interpolation for all default attributes
-	Dest.Vertex[VOut] = FMath::Lerp(V0Collection.Vertex[V0], V1Collection.Vertex[V1], T);
-	Dest.UV[VOut] = FMath::Lerp(V0Collection.UV[V0], V1Collection.UV[V1], T);
-	Dest.Color[VOut] = FMath::Lerp(V0Collection.Color[V0], V1Collection.Color[V1], T);
-	Dest.Normal[VOut] = FMath::Lerp(V0Collection.Normal[V0], V1Collection.Normal[V1], T).GetSafeNormal();
-	FVector TangentU = FMath::Lerp(V0Collection.TangentU[V0], V1Collection.TangentU[V1], T);
-	// don't lerp for TangentV, as it will be determined entirely by the lerp'd U and N
-	Dest.TangentV[VOut] = (Dest.Normal[VOut] ^ TangentU).GetSafeNormal();
-	Dest.TangentU[VOut] = (Dest.TangentV[VOut] ^ Dest.Normal[VOut]).GetSafeNormal();
-	// bone map value does not actually matter here; we will overwrite it later when we copy vertices out to new geometry groups
-	Dest.BoneMap[VOut] = V0Collection.BoneMap[V0];
-}
 
-void PLANARCUT_API ComputeTriangleNormals(
-	const TArrayView<const FVector> Vertices,
-	const TArrayView<const FIntVector> Triangles,
-	TArray<FVector>& TriangleNormals
-)
-{
-	TriangleNormals.SetNumUninitialized(Triangles.Num());
-	// TODO: parallel for?
-	for (int32 TriIdx = 0; TriIdx < Triangles.Num(); TriIdx++)
-	{
-		const FIntVector& Tri = Triangles[TriIdx];
-		FVector Edge1 = Vertices[Tri.Y] - Vertices[Tri.X];
-		FVector Edge2 = Vertices[Tri.Z] - Vertices[Tri.Y];
-		FVector Normal = FVector::CrossProduct(Edge2, Edge1);
-		Normal.Normalize();
-		TriangleNormals[TriIdx] = Normal;
-	}
-}
 
-// struct to store material info, e.g. to define what should go on to new faces along cut surfaces
-struct FaceMaterialInfo
-{
-	int32 MaterialID;
-	bool Visible;
-};
 
-// helper function to interpolate geometry collection data
-int32 AddInterpolatedVertex(const FGeometryCollection& Source, int32 SourceVertexNum, int32 V0, int32 V1, float T, FGeometryCollection& Dest,
-	TFunction<void(const FGeometryCollection&, int32, const FGeometryCollection&, int32, float, int32, FGeometryCollection&)> Interpolate)
+// Simpler invocation of CutWithPlanarCells w/ reasonable defaults
+int32 CutWithPlanarCells(
+	FPlanarCells& Cells,
+	FGeometryCollection& Source,
+	int32 TransformIdx,
+	double Grout,
+	double CollisionSampleSpacing,
+	const TOptional<FTransform>& TransformCollection,
+	bool bIncludeOutsideCellInOutput,
+	float CheckDistanceAcrossOutsideCellForProximity,
+	bool bSetDefaultInternalMaterialsFromCollection
+)
 {
-	int32 AddedVertexIdx = Dest.AddElements(1, FGeometryCollection::VerticesGroup);
-
-	const FGeometryCollection *V0Data = &Source;
-	const FGeometryCollection *V1Data = &Source;
-	if (V0 >= SourceVertexNum)
-	{
-		V0 -= SourceVertexNum;
-		V0Data = &Dest;
-	}
-	if (V1 >= SourceVertexNum)
-	{
-		V1 -= SourceVertexNum;
-		V1Data = &Dest;
-	}
-
-	Interpolate(*V0Data, V0, *V1Data, V1, T, AddedVertexIdx, Dest);
-
-	return AddedVertexIdx;
+	TArray<int32> TransformIndices { TransformIdx };
+	return CutMultipleWithPlanarCells(Cells, Source, TransformIndices, Grout, CollisionSampleSpacing, TransformCollection, bIncludeOutsideCellInOutput, CheckDistanceAcrossOutsideCellForProximity, bSetDefaultInternalMaterialsFromCollection);
 }
 
 
-
-// Output structure, stores one mesh per cell (including the "Outside of complex" cell, if needed)
-// (currently for internal use; not exposed)
-struct OutputCells
+int32 CutMultipleWithMultiplePlanes(
+	const TArrayView<const FPlane>& Planes,
+	FInternalSurfaceMaterials& InternalSurfaceMaterials,
+	FGeometryCollection& Collection,
+	const TArrayView<const int32>& TransformIndices,
+	double Grout,
+	double CollisionSampleSpacing,
+	const TOptional<FTransform>& TransformCollection,
+	bool bSetDefaultInternalMaterialsFromCollection
+)
 {
-	FGeometryCollection AddedVerticesCollection;
-	TArray<TArray<FIntVector>> CellTriangles;
-	TArray<TArray<int32>> CellTriangleSources; // indices into the original GeometryCollection face arrays (for copying over face attrib data)
-	TArray<TArray<int32>> CellVertexMapping; // indices into the original + added vertex arrays (for copying over vertex attrib data)
-	TArray<TArray<int32>> NeighborCells; // indices of output cells that neighbor this cell
-	int32 NoCellIdx = -1;
-
-	/**
-	 * Get index of the "outside" cell, for space that is classified as outside of all cells
-	 * If that index does not exist yet, create it.
-	 */
-	int32 GetNoCellIdx()
-	{
-		if (NoCellIdx == -1)
-		{
-			NoCellIdx = CellTriangles.Num();
-			CellTriangleSources.SetNum(NoCellIdx + 1);
-			CellVertexMapping.SetNum(NoCellIdx + 1);
-			CellTriangles.SetNum(NoCellIdx + 1);
-			NeighborCells.SetNum(NoCellIdx + 1);
-		}
-		return NoCellIdx;
-	}
-
-	void ConnectCells(int32 CellA, int32 CellB)
-	{
-		if (INDEX_NONE == NeighborCells[CellA].Find(CellB))
-		{
-			NeighborCells[CellA].Add(CellB);
-			NeighborCells[CellB].Add(CellA);
-		}
-		else
-		{
-			// cells should be symmetrically connected, so if the cell was already connected A->B, it should have been connected B->A
-			ensure(INDEX_NONE != NeighborCells[CellB].Find(CellA));
-		}
-	}
-
-	inline int32 OutputCellIdx(int32 CellID)
-	{
-		if (CellID < 0)
-		{
-			return GetNoCellIdx();
-		}
-		else
-		{
-			return CellID;
-		}
-	}
-
-	int32 NumNonEmptyCells() const
-	{
-		int32 NonEmptyCells = 0;
-		for (int32 CellIdx = 0; CellIdx < Num(); CellIdx++)
-		{
-			if (CellTriangles[CellIdx].Num() > 0)
-			{
-				NonEmptyCells++;
-			}
-		}
-		return NonEmptyCells;
-	}
-
-	OutputCells(int32 NumCells)
-	{
-		CellTriangleSources.SetNum(NumCells);
-		CellVertexMapping.SetNum(NumCells);
-		CellTriangles.SetNum(NumCells);
-		NeighborCells.SetNum(NumCells);
-		NoCellIdx = -1;
-	}
-	
-	int32 Num() const
-	{
-		check(CellVertexMapping.Num() == CellTriangles.Num());
-		return CellTriangles.Num();
-	}
+	int32 OrigNumGeom = Collection.FaceCount.Num();
+	int32 CurNumGeom = OrigNumGeom;
 
-	inline void AddTriangle(int32 CellIdx, int32 SourceTriangleIdx, FIntVector Triangle)
+	if (bSetDefaultInternalMaterialsFromCollection)
 	{
-		CellTriangles[CellIdx].Add(Triangle);
-		CellTriangleSources[CellIdx].Add(SourceTriangleIdx);
+		InternalSurfaceMaterials.SetUVScaleFromCollection(Collection);
 	}
 
-
-	int32 AddToGeometryCollection(FGeometryCollection& Source, const FInternalSurfaceMaterials &InternalMaterial, bool bIncludeOutsideCellInOutput, int32 SourceVertexNumWhenCut, int32 TransformParent = -1, int32 OverrideGlobalMaterialID = -1) const
+	if (!Collection.HasAttribute("Proximity", FGeometryCollection::GeometryGroup))
 	{
-		int32 InternalMaterialID = OverrideGlobalMaterialID > -1 ? OverrideGlobalMaterialID : InternalMaterial.GlobalMaterialID;
-		FGeometryCollection& Output = Source;  // keep separate variable names for Source and Output just in case we want to re-purpose this code later to add to a different collection
-
-		bool bHasProximity = Source.HasAttribute("Proximity", FGeometryCollection::GeometryGroup);
-
-		const FGeometryCollection& AddedVertices = AddedVerticesCollection;
-		int32 NewGeometryStartIdx = Output.FaceStart.Num();
-		int32 SourceVertexNum = Source.Vertex.Num();
-		int32 SourceFaceNum = Source.Indices.Num();
-
-		int32 TotalVerticesAdded = 0;
-		int32 TotalFacesAdded = 0;
-		int32 NumNewGeometries = NumNonEmptyCells();
-		int NumCellsToDump = Num();
-		// if we have a no-cell index set, and we're not including the outside geom, then skip that cell
-		if (!bIncludeOutsideCellInOutput && NoCellIdx > -1)
-		{
-			NumCellsToDump--;
-			ensure(NoCellIdx == NumCellsToDump); // by convention the index of the "no cell" geometry is the last index
-			if (CellTriangles[NoCellIdx].Num())
-			{
-				NumNewGeometries--;
-			}
-		}
-
-		int32 GeometryStart = Output.AddElements(NumNewGeometries, FGeometryCollection::GeometryGroup);
-		int32 TransformsStart = Output.AddElements(NumNewGeometries, FGeometryCollection::TransformGroup);
-		int32 GeometrySubIdx = 0;
-
-		TArray<int32> CellIdxToGeometryIdxMap;
-		if (bHasProximity)
-		{
-			CellIdxToGeometryIdxMap.SetNum(Num());
-			for (int32 Idx = 0, N = Num(); Idx < N; Idx++)
-			{
-				CellIdxToGeometryIdxMap[Idx] = -1;
-			}
-		}
-
-		TArray<FTranslationMatrix> ChildInverseTransforms;
-		for (int32 OutputCellIdx = 0; OutputCellIdx < NumCellsToDump; OutputCellIdx++)
-		{
-			int32 NumTriangles = CellTriangles[OutputCellIdx].Num();
-			int32 NumVertices = CellVertexMapping[OutputCellIdx].Num();
-			if (NumTriangles > 0)
-			{
-				int32 GeometryIdx = GeometryStart + GeometrySubIdx;
-				Output.FaceCount[GeometryIdx] = NumTriangles;
-				Output.FaceStart[GeometryIdx] = SourceFaceNum + TotalFacesAdded;
-				Output.VertexStart[GeometryIdx] = SourceVertexNum + TotalVerticesAdded;
-				Output.VertexCount[GeometryIdx] = NumVertices;
-				int32 TransformIdx = TransformsStart + GeometrySubIdx;
-				Output.TransformIndex[GeometryIdx] = TransformIdx;
-				Output.TransformToGeometryIndex[TransformIdx] = GeometryIdx;
-				if (TransformParent > -1)
-				{
-					// TODO: this is probably not the best way to build the bone name string?
-					Output.BoneName[TransformIdx] = Output.BoneName[TransformParent] + "_" + FString::FromInt(GeometrySubIdx);
-					Output.BoneColor[TransformIdx] = Output.BoneColor[TransformParent];
-					Output.Parent[TransformIdx] = TransformParent;
-					Output.Children[TransformParent].Add(TransformIdx);
-					Output.SimulationType[TransformParent] = FGeometryCollection::ESimulationTypes::FST_Clustered;
-				}
-
-
-				// Set the transform for the child geometry
-				// Note to make it easier to procedurally texture later, it's better to leave it in the same space
-				Output.Transform[TransformIdx] = FTransform::Identity;
-				Output.SimulationType[TransformIdx] = FGeometryCollection::ESimulationTypes::FST_Rigid;
-				ChildInverseTransforms.Emplace(FVector::ZeroVector);
-
-				GeometrySubIdx++;
-			}
-			TotalVerticesAdded += CellVertexMapping[OutputCellIdx].Num();
-			TotalFacesAdded += CellTriangles[OutputCellIdx].Num();
-		}
-		int32 VerticesStart = Output.AddElements(TotalVerticesAdded, FGeometryCollection::VerticesGroup);
-		int32 FacesStart = Output.AddElements(TotalFacesAdded, FGeometryCollection::FacesGroup);
-
-		int32 VertexGroupStart = VerticesStart;
-		int32 FaceGroupStart = FacesStart;
-		GeometrySubIdx = 0;
-		for (int32 OutputCellIdx = 0; OutputCellIdx < NumCellsToDump; OutputCellIdx++)
-		{
-			int32 NumTriangles = CellTriangles[OutputCellIdx].Num();
-			int32 NumVertices = CellVertexMapping[OutputCellIdx].Num();
-			int32 GeometryIdx = GeometryStart + GeometrySubIdx;
-			if (bHasProximity)
-			{
-				CellIdxToGeometryIdxMap[OutputCellIdx] = GeometryIdx;
-			}
-			FTranslationMatrix ToLocal(FVector(0));
-			if (NumTriangles > 0)
-			{
-				ToLocal = ChildInverseTransforms[GeometrySubIdx];
-				GeometrySubIdx++;
-			}
-			for (int32 VertexSubIdx = 0; VertexSubIdx < NumVertices; VertexSubIdx++)
-			{
-				int32 CopyVertexIdx = CellVertexMapping[OutputCellIdx][VertexSubIdx];
-				const FGeometryCollection* CopyFromCollection = &Source;
-				if (CopyVertexIdx >= SourceVertexNumWhenCut)
-				{
-					CopyFromCollection = &AddedVertices;
-					CopyVertexIdx -= SourceVertexNumWhenCut;
-				}
-				int32 CopyToIdx = VertexGroupStart + VertexSubIdx;
-				// TODO: add a function to GeometryCollection api to copy all properties of an element automatically, rather than hard-coding this?
-				Output.Vertex[CopyToIdx]		= ToLocal.TransformPosition(CopyFromCollection->Vertex[CopyVertexIdx]);
-				Output.Normal[CopyToIdx]        = ToLocal.TransformVector(CopyFromCollection->Normal[CopyVertexIdx]);
-				Output.UV[CopyToIdx]			= CopyFromCollection->UV[CopyVertexIdx];
-				Output.TangentU[CopyToIdx]		= ToLocal.TransformVector(CopyFromCollection->TangentU[CopyVertexIdx]);
-				Output.TangentV[CopyToIdx]		= ToLocal.TransformVector(CopyFromCollection->TangentV[CopyVertexIdx]);
-				Output.Color[CopyToIdx]			= CopyFromCollection->Color[CopyVertexIdx];
-
-				// Bone map should actually be set based on the transform of the new geometry, not copied from the old vertex
-				Output.BoneMap[CopyToIdx] = Output.TransformIndex[GeometryIdx];
-			}
-			for (int32 FaceSubIdx = 0; FaceSubIdx < CellTriangles[OutputCellIdx].Num(); FaceSubIdx++)
-			{
-				int32 CopyToIdx = FaceGroupStart + FaceSubIdx;
-				int32 SourceIdx = CellTriangleSources[OutputCellIdx][FaceSubIdx];
-				if (SourceIdx > -1) // we know the source face; copy information from there
-				{
-					// TODO: add a function to GeometryCollection api to copy all properties of an element automatically, rather than hard-coding this?
-					Output.Visible[CopyToIdx] = Source.Visible[SourceIdx];
-					// Note: I don't worry about Output.MaterialIndex[CopyToIdx], this will need to be rebuilt regardless later.
-					Output.MaterialID[CopyToIdx] = Source.MaterialID[SourceIdx];
-				}
-				else
-				{
-					Output.Visible[CopyToIdx] = InternalMaterial.bGlobalVisibility;
-					// Note: I don't worry about Output.MaterialIndex[CopyToIdx], this will need to be rebuilt regardless later.
-					Output.MaterialID[CopyToIdx] = InternalMaterialID;
-				}
-
-				// Face indices are the one property that isn't just copied over blindly; put the correct value *after* filling the rest of the data, in case we want to change the above to blindly copy all attributes
-				Output.Indices[CopyToIdx] = CellTriangles[OutputCellIdx][FaceSubIdx] + FIntVector(VertexGroupStart);
-			}
-			VertexGroupStart += NumVertices;
-			FaceGroupStart += CellTriangles[OutputCellIdx].Num();
-		}
-
-		if (bHasProximity)
-		{
-			TManagedArray<TSet<int32>>& Proximity = Source.GetAttribute<TSet<int32>>("Proximity", FGeometryCollection::GeometryGroup);
-
-			for (int32 OutputCellIdx = 0; OutputCellIdx < NumCellsToDump; OutputCellIdx++)
-			{
-				int32 GeomAIdx = CellIdxToGeometryIdxMap[OutputCellIdx];
-				if (GeomAIdx == -1)
-				{
-					continue;
-				}
-				for (int32 ConnectedCellIdx : NeighborCells[OutputCellIdx])
-				{
-					int32 GeomBIdx = CellIdxToGeometryIdxMap[ConnectedCellIdx];
-					if (GeomBIdx == -1)
-					{
-						continue;
-					}
-					Proximity[GeomAIdx].Add(GeomBIdx);
-					// TODO: make the NeighborCells thing be one-sided and do symmetric adds here
-				}
-			}
-		}
-
-		if (Source.BoundingBox.Num())
-		{
-			// Initialize BoundingBox
-			for (int32 Idx = GeometryStart; Idx < Source.BoundingBox.Num(); ++Idx)
-			{
-				Source.BoundingBox[Idx].Init();
-			}
-
-			// Compute BoundingBox
-			for (int32 Idx = SourceVertexNum; Idx < Source.Vertex.Num(); ++Idx)
-			{
-				int32 TransformIndexValue = Source.BoneMap[Idx];
-				Source.BoundingBox[Source.TransformToGeometryIndex[TransformIndexValue]] += Source.Vertex[Idx];
-			}
-		}
-
-		return NewGeometryStartIdx;
+		const FManagedArrayCollection::FConstructionParameters GeometryDependency(FGeometryCollection::GeometryGroup);
+		Collection.AddAttribute<TSet<int32>>("Proximity", FGeometryCollection::GeometryGroup, GeometryDependency);
 	}
-};
 
-/**
-* Transform local geometry, updating the corresponding transform so the shape itself is not changed
-*/
-void TransformLocalGeometry(FGeometryCollection& Source, int32 TransformIdx, const FTransform& Transform, const FTransform& InverseTransform)
-{
-	int GeometryIdx = Source.TransformToGeometryIndex[TransformIdx];
-	// recompute bounds (rather than directly transforming, so it remains a tight bound even if transform includes rotation)
-	FBox Bounds; Bounds.Init();
-	for (int32 VertIdx = Source.VertexStart[GeometryIdx], VertEnd = Source.VertexCount[GeometryIdx] + Source.VertexStart[GeometryIdx]; VertIdx < VertEnd; VertIdx++)
-	{
-		FVector Pos = Transform.TransformPosition(Source.Vertex[VertIdx]);
-		Bounds += Pos;
-		Source.Vertex[VertIdx] = Pos;
-	}
-	Source.BoundingBox[GeometryIdx] = Bounds;
-	Source.Transform[TransformIdx] = InverseTransform * Source.Transform[TransformIdx];
-}
+	FTransform CollectionToWorld = TransformCollection.Get(FTransform::Identity);
 
-/** 
- * Update a given transform w/ geometry to have vertices locally centered at the origin and positions not exceeding a -1,1 range, changing the transform appropriately so that the shape itself is not changed
- */
-FTransform CenterAndScaleLocalGeometry(FGeometryCollection& Source, int32 TransformIdx)
-{
-	int GeometryIdx = Source.TransformToGeometryIndex[TransformIdx];
-	if (!ensure(GeometryIdx != INDEX_NONE)) // transform had no geometry
-	{
-		return FTransform::Identity;
-	}
+	FDynamicMeshCollection MeshCollection(&Collection, TransformIndices, CollectionToWorld);
+	double OnePercentExtend = MeshCollection.Bounds.MaxDim() * .01;
 
-	FBox GeomBox; GeomBox.Init();
-	if (Source.BoundingBox.Num())
-	{
-		GeomBox = Source.BoundingBox[GeometryIdx];
-	}
-	if (!GeomBox.IsValid)
-	{
-		for (int32 VertIdx = Source.VertexStart[GeometryIdx], VertEnd = Source.VertexCount[GeometryIdx] + Source.VertexStart[GeometryIdx]; VertIdx < VertEnd; VertIdx++)
-		{
-			GeomBox += Source.Vertex[VertIdx];
-		}
-	}
+	int32 NewGeomStartIdx = -1;
+	NewGeomStartIdx = MeshCollection.CutWithMultiplePlanes(Planes, Grout, CollisionSampleSpacing, &Collection, InternalSurfaceMaterials, bSetDefaultInternalMaterialsFromCollection);
 
-	if (!ensure(GeomBox.IsValid)) // transform had corresponding geometry index but it had zero vertices?
-	{
-		return FTransform::Identity;
-	}
-	
-	FVector Center, Extents;
-	GeomBox.GetCenterAndExtents(Center, Extents);
-	float MaxExtent = Extents.GetMax();
-	float InvScaleFactor = MaxExtent < 1 ? 1 : MaxExtent;
-	float ScaleFactor = 1.0f / InvScaleFactor;
-	FTransform CenterAndFit, InverseCenterAndFit;
-	CenterAndFit.SetTranslationAndScale3D(-Center*ScaleFactor, FVector(ScaleFactor, ScaleFactor, ScaleFactor));
-	InverseCenterAndFit.SetTranslationAndScale3D(Center, FVector(InvScaleFactor, InvScaleFactor, InvScaleFactor));
-
-	TransformLocalGeometry(Source, TransformIdx, CenterAndFit, InverseCenterAndFit);
-
-	return InverseCenterAndFit;
+	Collection.ReindexMaterials();
+	return NewGeomStartIdx;
 }
 
 
-
-//// useful for debugging: code to dump a single geometry to obj
-//#include <fstream>
-//#include <string>
-//
-//void WriteGeom(const FString& Path, FGeometryCollection &Collection, int32 GeometryIdx)
-//{
-//	std::ofstream f(std::string(TCHAR_TO_UTF8(*Path)));
-//	for (FVector &V : Collection.Vertex)
-//	{
-//		f << "v " << V.X << " " << V.Y << " " << V.Z << std::endl;
-//	}
-//
-//	int32 Start = Collection.FaceStart[GeometryIdx];
-//	for (int32 i = Start; i < Start + Collection.FaceCount[GeometryIdx]; i++)
-//	{
-//		FIntVector Tri = Collection.Indices[i] + FIntVector(1);
-//		f << "f " << Tri.X << " " << Tri.Y << " " << Tri.Z << std::endl;
-//	}
-//	f.close();
-//}
-//
-//void WritePlanarCells(const FString& Path, FPlanarCells& Cells, const TArrayView<const FVector>& PlaneBoundaryVertices)
-//{
-//	std::ofstream f(std::string(TCHAR_TO_UTF8(*Path)));
-//	for (const FVector &V : PlaneBoundaryVertices)
-//	{
-//		f << "v " << V.X << " " << V.Y << " " << V.Z << std::endl;
-//	}
-//
-//	
-//	for (int32 i = 0; i < Cells.PlaneBoundaries.Num(); i++)
-//	{
-//		TArray<int32>& Inds = Cells.PlaneBoundaries[i];
-//		for (int32 j = 1; j + 1 < Inds.Num(); j++)
-//		{
-//			f << "f " << Inds[0] + 1 << " " << Inds[j] + 1 << " " << Inds[j + 1] + 1 << std::endl;
-//		}
-//	}
-//	f.close();
-//}
-
-/**
- * Helper version of CutWithPlanarCells; exposes more parameters for later experimentation
- * Cut a (subset of a) GeometryCollection with PlanarCells, and add each cut cell to an OutputCells struct
- */
-void CutWithPlanarCellsHelper(
-	const FPlanarCells &Cells,
-	const TArrayView<const FPlane> TransformedPlanes,
-	const TArrayView<const FVector> TransformedPlaneBoundaryVertices,
-	const FTransform& LocalSpaceToPlanarCellSpace,
-	const FGeometryCollection& Source,
-	int32 GeometryIdx,
-	int32 TriangleStart,
-	int32 NumTriangles,
-	const TArrayView<const FVector> TriangleNormals,
-	double PlaneEps,
+// Cut multiple Geometry groups inside a GeometryCollection with PlanarCells, and add each cut cell back to the GeometryCollection as a new child of their source Geometry
+int32 CutMultipleWithPlanarCells(
+	FPlanarCells& Cells,
+	FGeometryCollection& Source,
+	const TArrayView<const int32>& TransformIndices,
+	double Grout,
+	double CollisionSampleSpacing,
+	const TOptional<FTransform>& TransformCollection,
+	bool bIncludeOutsideCellInOutput,
 	float CheckDistanceAcrossOutsideCellForProximity,
-	TFunction<void(const FGeometryCollection&, int32, const FGeometryCollection&, int32, float, int32, FGeometryCollection&)> Interpolate,
-	OutputCells &Output,
-	const FInternalSurfaceMaterials *InternalMaterials = nullptr, // optional override for cells' internal materials
-	TMeshAABBTree3<FGeometryCollectionMeshAdapter> *PrecomputedAABBTree = nullptr // optionally pass in an already-computed AABB tree
+	bool bSetDefaultInternalMaterialsFromCollection
 )
 {
-	if (!InternalMaterials)
+	if (!Source.HasAttribute("Proximity", FGeometryCollection::GeometryGroup))
 	{
-		InternalMaterials = &Cells.InternalSurfaceMaterials;
+		const FManagedArrayCollection::FConstructionParameters GeometryDependency(FGeometryCollection::GeometryGroup);
+		Source.AddAttribute<TSet<int32>>("Proximity", FGeometryCollection::GeometryGroup, GeometryDependency);
 	}
-	// shorthand accessor for vertices in the source geometry collection
-	const TArrayView<const FVector> Vertices(Source.Vertex.GetData(), Source.Vertex.Num());
-	const TArrayView<const FIntVector> Triangles(Source.Indices.GetData() + TriangleStart, NumTriangles);
 
-	// consider trade-offs between cases where we could have a more consistent mesh vs having simpler processing / fewer triangles
-	constexpr bool bCareAboutTJunctionsEvenALittleBit = false;
-	bool bNoiseOnPlane = InternalMaterials->NoiseSettings.IsSet();
-
-	// extract an average scale for this transform to support properly spacing noise points, if requested
-	float AverageGlobalScale = 1;
-	if (bNoiseOnPlane)
+	if (bSetDefaultInternalMaterialsFromCollection)
 	{
-		FTransform LocalToGlobalTransform = GeometryCollectionAlgo::GlobalMatrix(Source.Transform, Source.Parent, Source.TransformIndex[GeometryIdx]);
-		FVector Scales = LocalToGlobalTransform.GetScale3D();
-		AverageGlobalScale = FMath::Max(KINDA_SMALL_NUMBER, FVector::DotProduct(Scales.GetAbs(), FVector(1. / 3.)));
+		Cells.InternalSurfaceMaterials.SetUVScaleFromCollection(Source);
 	}
-	float AverageGlobalScaleInv = 1.0 / AverageGlobalScale;
-
-	struct PlaneFrame
-	{
-		FVector3d Origin, X, Y;
 
-		PlaneFrame(FVector3d Origin, FVector3d Normal) : Origin(Origin)
-		{
-			VectorUtil::MakePerpVectors(Normal, X, Y);
-		}
+	FTransform CollectionToWorld = TransformCollection.Get(FTransform::Identity);
 
-		inline FVector2d Project(const FVector3d& Pt)
-		{
-			FVector3d RelPt = Pt - Origin;
-			return FVector2d(RelPt.Dot(X), RelPt.Dot(Y));
-		}
+	FDynamicMeshCollection MeshCollection(&Source, TransformIndices, CollectionToWorld);
+	double OnePercentExtend = MeshCollection.Bounds.MaxDim() * .01;
+	FCellMeshes CellMeshes(Cells, MeshCollection.Bounds, Grout, OnePercentExtend, bIncludeOutsideCellInOutput);
 
-		inline FVector3d UnProject(const FVector2d& Pt)
-		{
-			return Origin + X * Pt.X + Y * Pt.Y;
-		}
-	};
+	int32 NewGeomStartIdx = -1;
 
-	TArray<PlaneFrame> PlaneFrames;
-	PlaneFrames.Reserve(Cells.Planes.Num());
-	for (int32 PlaneIdx = 0; PlaneIdx < Cells.Planes.Num(); PlaneIdx++)
-	{
-		const FPlane& Plane = TransformedPlanes[PlaneIdx];
-		FVector3d Normal(Plane.X, Plane.Y, Plane.Z);
-		const TArray<int32>& Boundary = Cells.PlaneBoundaries[PlaneIdx];
-		int32 NumBoundary = Boundary.Num();
-		if (NumBoundary)
-		{
-			PlaneFrames.Emplace(TransformedPlaneBoundaryVertices[Boundary[0]], Normal);
-		}
-		else
-		{
-			PlaneFrames.Emplace(Plane.W * Normal, Normal);
-		}
-	}
-
-	auto IsProjectionInsideBoundary = [&](const FVector& Pt, int32 PlaneIdx)
-	{
-		const TArray<int32>& Boundary = Cells.PlaneBoundaries[PlaneIdx];
-		int32 NumBoundary = Boundary.Num();
-		if (!NumBoundary) // unbounded plane case
-		{
-			return true;
-		}
-		TArray<FVector2d> ProjV;
-		ProjV.Reserve(NumBoundary);
-		for (int32 VIdx : Boundary)
-		{
-			ProjV.Add(PlaneFrames[PlaneIdx].Project(TransformedPlaneBoundaryVertices[VIdx]));
-		}
-		FPolygon2d Polygon(ProjV);
-		return Polygon.Contains(PlaneFrames[PlaneIdx].Project(Pt));
-	};
-
-	FGeometryCollectionMeshAdapter Adapter{ &Source, GeometryIdx };
-	TMeshAABBTree3<FGeometryCollectionMeshAdapter> LocalAABBTree;
-	TMeshAABBTree3<FGeometryCollectionMeshAdapter> *AABBTree = PrecomputedAABBTree;
-	if (!PrecomputedAABBTree)
-	{
-		LocalAABBTree.SetMesh(&Adapter);
-		AABBTree = &LocalAABBTree;
-	}
-	check(AABBTree->GetMesh()->Collection == &Source && AABBTree->GetMesh()->GeometryIdx == GeometryIdx); // verify that the AABBTree is looking at the correct geometry
-	TFastWindingTree<FGeometryCollectionMeshAdapter> FastWindingTree(AABBTree);
-
-	int32 InputVertexCount = Vertices.Num();
-	int32 NumPlanes = TransformedPlanes.Num();
-
-	TArray<TArray<int32>> PlanesThroughTriangle, PlanesOnTriangle; // PlanesThroughTriangle[TriIdx] -> All PlaneIdx crossing triangle
-	PlanesThroughTriangle.SetNum(Triangles.Num());
-	PlanesOnTriangle.SetNum(Triangles.Num());
-
-	// Mappings from PlaneIdx -> Elements (triangles, edges) on that plane
-	TArray<TArray<int32>> TrianglesOnPlane; // note: We will fill this later, using cut-down final triangles, so that no triangle ends up on two planes
-	TrianglesOnPlane.SetNum(NumPlanes);
-	TArray<TArray<TPair<FVector, FVector>>> EdgesOnPlane;
-	EdgesOnPlane.SetNum(NumPlanes);
-
-	// ~~~ PHASE 1: FIND COLLISIONS BETWEEN ALL PLANAR FACETS AND TRIANGLES ~~~	
-
-	// Obstacles to making this parallel: PlanesThroughTriangle, PlanesOnTriangle can both be edited by multiple planes.  If we make a per-triangle lock, could be fine!
-	for (int32 PlaneIdx = 0; PlaneIdx < NumPlanes; PlaneIdx++)
-	{
-		const FPlane &Plane = TransformedPlanes[PlaneIdx];
-		FBox BoundingBox(EForceInit::ForceInit);
-		if (Cells.PlaneBoundaries[PlaneIdx].Num())
-		{
-			for (int32 PlaneBoundaryVertexIdx : Cells.PlaneBoundaries[PlaneIdx])
-			{
-				BoundingBox += TransformedPlaneBoundaryVertices[PlaneBoundaryVertexIdx];
-			}
-		}
-
-		FAxisAlignedBox3d PlaneFacetBoxPlusEps(BoundingBox);
-		PlaneFacetBoxPlusEps.Max = PlaneFacetBoxPlusEps.Max + PlaneEps;
-		PlaneFacetBoxPlusEps.Min = PlaneFacetBoxPlusEps.Min - PlaneEps;
-		TMeshAABBTree3<FGeometryCollectionMeshAdapter>::FTreeTraversal TraverseNearPlane;
-		TraverseNearPlane.NextBoxF = [&](const FAxisAlignedBox3d& Box, int Depth)
-		{
-			if (BoundingBox.IsValid && !Box.Intersects(PlaneFacetBoxPlusEps))
-			{
-				return false;
-			}
-			int sides[3] = { 0,0,0 };
-			sides[PlaneSide(Plane, (FVector)Box.Min) + 1]++;
-			sides[PlaneSide(Plane, (FVector)Box.Max) + 1]++;
-			sides[PlaneSide(Plane, FVector(Box.Max.X, Box.Min.Y, Box.Min.Z)) + 1]++;
-			sides[PlaneSide(Plane, FVector(Box.Min.X, Box.Max.Y, Box.Min.Z)) + 1]++;
-			sides[PlaneSide(Plane, FVector(Box.Max.X, Box.Max.Y, Box.Min.Z)) + 1]++;
-			sides[PlaneSide(Plane, FVector(Box.Min.X, Box.Min.Y, Box.Max.Z)) + 1]++;
-			sides[PlaneSide(Plane, FVector(Box.Max.X, Box.Min.Y, Box.Max.Z)) + 1]++;
-			sides[PlaneSide(Plane, FVector(Box.Min.X, Box.Max.Y, Box.Max.Z)) + 1]++;
-			// we cross box if any vertex 'on' plane or vertices on both sides
-			return sides[1] || (sides[0] && sides[2]);
-		};
-		//for (int32 TriIdx = 0; TriIdx < NumTriangles; TriIdx++)
-		TraverseNearPlane.NextTriangleF = [&](int TriIdxInt) 
-		{
-			int32 TriIdx = TriIdxInt; // cast to int32 up top, because we use this a lot as a TArray index
-
-			const FIntVector &Tri = Triangles[TriIdx];
-
-			// TODO: further filter out cases where triangle doesn't intersect planar facet
-
-			double SX = PlaneDotDouble(Plane,Vertices[Tri.X]), SY = PlaneDotDouble(Plane, Vertices[Tri.Y]), SZ = PlaneDotDouble(Plane, Vertices[Tri.Z]);
-			int32 SXSide = SX < -PlaneEps ? -1 : SX > PlaneEps ? 1 : 0;
-			int32 SYSide = SY < -PlaneEps ? -1 : SY > PlaneEps ? 1 : 0;
-			int32 SZSide = SZ < -PlaneEps ? -1 : SZ > PlaneEps ? 1 : 0;
-			if (SXSide * SYSide == -1 || SYSide * SZSide == -1 || SZSide * SXSide == -1) // triangle crosses plane
-			{
-				PlanesThroughTriangle[TriIdx].Add(PlaneIdx);
-				int32 CrossIdx = 0;
-				FVector CrossPosns[2];
-				
-
-				auto AddCross = [&](float SDA, float SDB, int32 ASide, int32 BSide, int32 AIdx, int32 BIdx, int32& CrossIdxRef)
-				{
-					if (CrossIdxRef < 2 && ASide*BSide == -1)
-					{
-						double T = SDA / (SDA - SDB);
-						CrossPosns[CrossIdxRef] = (1 - T) * Vertices[AIdx] + T * Vertices[BIdx];
-						CrossIdxRef++;
-					}
-				};
-				AddCross(SX, SY, SXSide, SYSide, Tri.X, Tri.Y, CrossIdx);
-				AddCross(SY, SZ, SYSide, SZSide, Tri.Y, Tri.Z, CrossIdx);
-				AddCross(SZ, SX, SZSide, SXSide, Tri.Z, Tri.X, CrossIdx);
-				if (CrossIdx == 1) // One of the vertices is on the plane
-				{
-					int32 MinSDIdx = Tri.X;
-					float MinSDAbs = FMath::Abs(SX);
-					float SYAbs = FMath::Abs(SY);
-					if (SYAbs < MinSDAbs)
-					{
-						MinSDIdx = Tri.Y;
-						MinSDAbs = SYAbs;
-					}
-					float SZAbs = FMath::Abs(SZ);
-					if (SZAbs < MinSDAbs)
-					{
-						MinSDIdx = Tri.Z;
-						MinSDAbs = SZAbs;
-					}
-					CrossPosns[1] = Vertices[MinSDIdx];
-					CrossIdx++;
-				}
-				ensure(CrossIdx == 2);
-				EdgesOnPlane[PlaneIdx].Add(TPair<FVector, FVector>(CrossPosns[0], CrossPosns[1]));
-			}
-			else {
-				int32 OnPlaneCount =
-					(FMath::Abs(SX) < PlaneEps ? 1 : 0) +
-					(FMath::Abs(SY) < PlaneEps ? 1 : 0) +
-					(FMath::Abs(SZ) < PlaneEps ? 1 : 0);
-				if (OnPlaneCount == 3)
-				{
-					PlanesOnTriangle[TriIdx].Add(PlaneIdx);
-				}
-				else if (OnPlaneCount == 2)
-				{
-					int32 OnIdx = 0;
-					FVector OnPosns[2];
-
-					auto AddOn = [&Vertices, &PlaneEps, &OnPosns](float SD, int32 VIdx, int32 &OnIdxRef)
-					{
-						if (OnIdxRef < 2 && FMath::Abs(SD) < PlaneEps)
-						{
-							OnPosns[OnIdxRef] = Vertices[VIdx];
-							OnIdxRef++;
-						}
-					};
-					AddOn(SX, Tri.X, OnIdx);
-					AddOn(SY, Tri.Y, OnIdx);
-					AddOn(SZ, Tri.Z, OnIdx);
-					ensure(OnIdx == 2);
-					EdgesOnPlane[PlaneIdx].Add(TPair<FVector, FVector>(OnPosns[0], OnPosns[1]));
-				}
-			}
-		};
-		AABBTree->DoTraversal(TraverseNearPlane);
-	}
-
-	// ~~~ PHASE 2: CUT ALL TRIANGLES THAT CROSS PLANAR FACETS ~~~	
-	// Obstacles to making this parallel: CompletedEdgeSplits and AddedVertices are both read and edited from multiple triangles; 
-	//										CellFromPosition for FVoronoiDiagram is not thread safe
-	TMap<TPair<int32, int32>, int32> CompletedEdgeSplits; // for storing edge splits after they have already happened
-	int32 OrigTriNum = Triangles.Num();
-	
-	check(Output.Num() == Cells.NumCells);
-	FGeometryCollection& AddedVerticesCollection = Output.AddedVerticesCollection;
-	const TManagedArray<FVector> &AddedVertices = AddedVerticesCollection.Vertex;  // shorthand reference, only used for accesses
-
-	auto VertexPos = [&](int32 VertexIdx)
-	{
-		return VertexIdx < InputVertexCount ? Vertices[VertexIdx] : AddedVertices[VertexIdx - InputVertexCount];
-	};
-	auto Centroid = [&](const FIntVector &Tri)
-	{
-		return (VertexPos(Tri.X) + VertexPos(Tri.Y) + VertexPos(Tri.Z)) * (1.0 / 3.0);
-	};
-	for (int32 OrigTriIdx = 0; OrigTriIdx < OrigTriNum; OrigTriIdx++)
-	{
-		const FIntVector &OrigTri = Triangles[OrigTriIdx];
-		const FVector& TriNormal = TriangleNormals[OrigTriIdx];
-		FPlane TriPlane(TriNormal, FVector::DotProduct(TriNormal, Vertices[OrigTri.X]));
-
-		// TODO: consider a mode where we just handle every original triangle totally independently, with the caveat that we will duplicate vertices for every edge split w/ multiple triangles on the edge(!?)
-		if (!bCareAboutTJunctionsEvenALittleBit)
-		{
-			CompletedEdgeSplits.Empty();
-		}
-		TMap<TPair<int32, int32>, int32> InternalEdgeSplits; // Edge splits that are entirely inside the triangle don't need to be added to global map
-		TArray<FIntVector> AddedTriangles;  // Array of all triangles that will be added to cells based on this triangle
-		AddedTriangles.Add(OrigTri);		//	(initially just the original triangle)
-
-		// If there are any cutting planes, split up the triangle as needed
-		if (PlanesThroughTriangle[OrigTriIdx].Num() > 0)
-		{
-			for (int32 PlaneIdx : PlanesThroughTriangle[OrigTriIdx]) // fully cut one plane at a time
-			{
-				const FPlane& Plane = TransformedPlanes[PlaneIdx];
-				const FVector PlaneNormal(Plane.X, Plane.Y, Plane.Z);
-				FVector IntersectionDirection = TriNormal ^ PlaneNormal;
-				bool bFoundIntersectionLine = IntersectionDirection.Normalize(); // this should always succeed for anything in a PlanesThroughTriangle array ...
-
-				const TArray<int32>& PlaneBoundary = Cells.PlaneBoundaries[PlaneIdx];
-
-				// compute the extent of the planar facet on the line where it intersects the plane of the triangle.  used below to skip sub-triangles that don't touch the facet.
-				FInterval1d PlaneFacetInterval = FInterval1d::Empty();
-				int32 PlaneBoundaryNum = PlaneBoundary.Num();
-				if (bFoundIntersectionLine && PlaneBoundaryNum)
-				{
-					for (int32 Idx = 0, LastIdx = PlaneBoundaryNum - 1; Idx < PlaneBoundaryNum; LastIdx = Idx++)
-					{
-						double T;
-						const FVector& A = TransformedPlaneBoundaryVertices[PlaneBoundary[Idx]];
-						const FVector& B = TransformedPlaneBoundaryVertices[PlaneBoundary[LastIdx]];
-						IsSegmentCrossing(TriPlane, A, B, T);
-						if (T > -PlaneEps && T < 1+PlaneEps) {
-							PlaneFacetInterval.Contain(FVector::DotProduct(IntersectionDirection, FMath::Lerp(A, B, T)));
-						}
-					}
-
-					// grow interval by a tolerance
-					PlaneFacetInterval.Min -= PlaneEps;
-					PlaneFacetInterval.Max += PlaneEps;
-				}
-
-				// all added triangles need processing w/ the new possible cutting plane
-				TArray<int32> TriProcessQueue;
-				for (int32 ToProcessIdx = 0; ToProcessIdx < AddedTriangles.Num(); ToProcessIdx++)
-				{
-					TriProcessQueue.Add(ToProcessIdx);
-				}
-
-				TSet<int32> BoundaryVertices;
-				BoundaryVertices.Add(OrigTri.X);
-				BoundaryVertices.Add(OrigTri.Y);
-				BoundaryVertices.Add(OrigTri.Z);
-
-				auto DoEdgeSplit = [&](int32 V0, int32 V1, int32 VOff, int32 CurTriIdx, int32 ConsiderPlaneIdx)
-				{
-					int32 VSmall = V0, VBig = V1;
-					if (VSmall > VBig)
-					{
-						Swap(VSmall, VBig);
-					}
-					TPair<int32, int32> Edge(VSmall, VBig);
-
-					bool BoundaryEdge = BoundaryVertices.Contains(V0) && BoundaryVertices.Contains(V1);
-
-					auto MakeTwoTriangles = [&](int32 SplitIdx)
-					{
-						AddedTriangles[CurTriIdx] = FIntVector(V0, SplitIdx, VOff);
-						int32 NewTriIdx = AddedTriangles.Emplace(SplitIdx, V1, VOff);
-						TriProcessQueue.Add(CurTriIdx);
-						TriProcessQueue.Add(NewTriIdx);
-					};
-
-					// check if this is a known split edge on boundary or internally
-					// TODO: do we even want to do this?  can we disable / remove this stuff if we don't care about making T Junctions?  (the other code will make a ton of them anyway)
-					// Partial answer: If we don't do this, we will repeat the edge split for every triangle on that edge, duplicating the vertex and disjoining the topology even if it's not a T Junction?
-					/*if (bCareAboutTJunctionsEvenALittleBit)
-					{*/
-						int32 *FoundBorder = BoundaryEdge ? CompletedEdgeSplits.Find(Edge) : nullptr;
-						if (FoundBorder)
-						{
-							BoundaryVertices.Add(*FoundBorder);
-							MakeTwoTriangles(*FoundBorder);
-							return true;
-						}
-						int32 *FoundInternal = !BoundaryEdge ? InternalEdgeSplits.Find(Edge) : nullptr;
-						if (FoundInternal)
-						{
-							MakeTwoTriangles(*FoundInternal);
-							return true;
-						}
-					//}
-
-					{
-						double t;
-						FVector P0 = VertexPos(V0), P1 = VertexPos(V1);
-						if (IsSegmentCrossing(TransformedPlanes[ConsiderPlaneIdx], P0, P1, t))
-						{
-							int32 SplitVert = InputVertexCount + AddInterpolatedVertex(Source, Vertices.Num(), V0, V1, t, AddedVerticesCollection, Interpolate);
-							if (BoundaryEdge)
-							{
-								BoundaryVertices.Add(SplitVert);
-								CompletedEdgeSplits.Add(Edge, SplitVert);
-							}
-							else
-							{
-								InternalEdgeSplits.Add(Edge, SplitVert);
-							}
-							MakeTwoTriangles(SplitVert);
-							return true;
-						}
-					}
-
-					return false;
-				};
-
-				while (TriProcessQueue.Num())
-				{
-					int32 TriToSplitIdx = TriProcessQueue.Pop();
-					const FIntVector& Tri = AddedTriangles[TriToSplitIdx];
-					if (bFoundIntersectionLine && PlaneBoundaryNum > 0)
-					{
-						// check if the triangle overlaps with the the planar facet; if no overlap, no need to split
-						FInterval1d TriInterval = FInterval1d::Empty();
-						FVector A = VertexPos(Tri.X), B = VertexPos(Tri.Y), C = VertexPos(Tri.Z);
-						auto ContainCrossing = [&Plane, &PlaneEps, &IntersectionDirection](FInterval1d& Inter, const FVector& P0, const FVector& P1)
-						{
-							double T;
-							IsSegmentCrossing(Plane, P0, P1, T);
-							if (T > -PlaneEps && T < 1 + PlaneEps)
-							{
-								Inter.Contain(FVector::DotProduct(IntersectionDirection, FMath::Lerp(P0, P1, T)));
-							}
-						};
-						ContainCrossing(TriInterval, A, B);
-						ContainCrossing(TriInterval, B, C);
-						ContainCrossing(TriInterval, C, A);
-						if (!TriInterval.Overlaps(PlaneFacetInterval))
-						{
-							continue;
-						}
-					}
-					DoEdgeSplit(Tri.X, Tri.Y, Tri.Z, TriToSplitIdx, PlaneIdx)
-						|| DoEdgeSplit(Tri.Y, Tri.Z, Tri.X, TriToSplitIdx, PlaneIdx)
-						|| DoEdgeSplit(Tri.Z, Tri.X, Tri.Y, TriToSplitIdx, PlaneIdx);
-				}
-			}
-		}
-
-		const FVector &TriangleNormal = TriangleNormals[OrigTriIdx];
-		for (const FIntVector& Tri : AddedTriangles)
-		{
-			FVector TriPos = Centroid(Tri);
-			int32 Cell = Cells.CellFromPosition(LocalSpaceToPlanarCellSpace.TransformPosition(TriPos));
-
-			// triangles that were coplanar with a cutting plane can be re-assigned to the neighboring cell based on their normal
-			int32 OwnedByPlane = -1;
-			for (int32 PlaneIdx : PlanesOnTriangle[OrigTriIdx])
-			{
-				TPair<int32, int32> PlaneCells = Cells.PlaneCells[PlaneIdx];
-				if ((PlaneCells.Key == Cell || PlaneCells.Value == Cell) && IsProjectionInsideBoundary(TriPos, PlaneIdx))
-				{
-					OwnedByPlane = PlaneIdx;
-					// TODO: note if the original triangle was degenerate then this will be arbitrary, but not sure we care about that case!
-					FVector PlaneNormal(TransformedPlanes[PlaneIdx].X, TransformedPlanes[PlaneIdx].Y, TransformedPlanes[PlaneIdx].Z);
-					Cell = FVector::DotProduct(TriangleNormal, PlaneNormal) > 0 ? PlaneCells.Key : PlaneCells.Value;
-					break;
-				}
-			}
-			// Store plane ownership decision for later use by triangulation algo
-			if (OwnedByPlane > -1)
-			{
-				TrianglesOnPlane[OwnedByPlane].AddUnique(OrigTriIdx); // Another trouble spot if we want to parallel-for this whole section per triangle
-			}
-			if (Cell < 0)
-			{
-				// dump triangles that have no cell (e.g. this happens if a Voronoi diagram didn't enclose the whole mesh)
-				Cell = Output.GetNoCellIdx();
-			}
-
-			Output.AddTriangle(Cell, OrigTriIdx + TriangleStart, Tri);
-		}
-	}
-
-	// ~~~ TODO: OPTIONAL PHASE 3: One additional pass over all triangles, splitting any triangles w/ boundary edges in the split map! ~~~
-
-	// ~~~ PHASE 4: TRIANGULATE ALL PLANAR CUTTING SURFACES AND ADD NEW FACES TO OUTPUT STRUCT ~~~
-
-	struct PlaneTriangulationInfo
-	{
-		TArray<FVector> LocalVertices; // positions of vertices used in the planar triangulation (TODO: we could try to optimize / not compute this in all cases / only store added vertices, at the cost of more complex lookups later.  probably not worth it?)
-		TArray<FVector2D> LocalUVs;
-		//TArray<int32> VertexSources; // indices of LocalVertices in to the global Vertices; -1 indicates the vertex has no corresponding global vertex
-		
-		TArray<FIntVector> LocalIndices;  // triangle indices into the local VertexSources array
-	};
-	TArray<PlaneTriangulationInfo> PlaneTriangulations; PlaneTriangulations.SetNum(NumPlanes);
-
-	bool bNoParallel = false;
-	//for (int32 PlaneIdx = 0; PlaneIdx < NumPlanes; PlaneIdx++)
-	ParallelFor(NumPlanes, [&](int32 PlaneIdx)
-	{
-		const FPlane &Plane = TransformedPlanes[PlaneIdx];
-		const TArray<int32> BoundaryIndices = Cells.PlaneBoundaries[PlaneIdx];
-		PlaneTriangulationInfo& Triangulation = PlaneTriangulations[PlaneIdx];
-		int32 NumBoundary = BoundaryIndices.Num();
-		FVector PlaneNormal(Plane.X, Plane.Y, Plane.Z);
-		FVector Origin = (FVector)PlaneFrames[PlaneIdx].Origin;
-
-		// check if constrained Delaunay triangulation problem needed for plane (false if no geometry was touching the planar facet)
-		bool bAnyElementsOnPlane = EdgesOnPlane[PlaneIdx].Num() + TrianglesOnPlane[PlaneIdx].Num() > 0;
-		bool bConvexFacet = Cells.AssumeConvexCells;
-		bool bHasBoundary = NumBoundary > 2;
-		if (bAnyElementsOnPlane || (bHasBoundary && (bNoiseOnPlane || !bConvexFacet)))
-		{
-			FAxisAlignedBox2d Bounds2D = FAxisAlignedBox2d::Empty();
-			TArray<FVector2d> Boundary;
-			TArray<TPair<FVector2d, FVector2d>> PlanarEdges;
-			int32 BoundaryEdgeStart = -1;
-			for (const TPair<FVector, FVector>& Edge : EdgesOnPlane[PlaneIdx])
-			{
-				PlanarEdges.Emplace(PlaneFrames[PlaneIdx].Project(Edge.Key), PlaneFrames[PlaneIdx].Project(Edge.Value));
-			}
-			TArray<FVector2d> ProjectedTriVertices;
-			for (const int32 TriIdx : TrianglesOnPlane[PlaneIdx])
-			{
-				const FIntVector& Tri = Triangles[TriIdx];
-				FVector2d A = PlaneFrames[PlaneIdx].Project(Vertices[Tri.X]), B = PlaneFrames[PlaneIdx].Project(Vertices[Tri.Y]), C = PlaneFrames[PlaneIdx].Project(Vertices[Tri.Z]);
-				PlanarEdges.Emplace(A, B);
-				PlanarEdges.Emplace(B, C);
-				PlanarEdges.Emplace(C, A);
-				ProjectedTriVertices.Add(A);
-				ProjectedTriVertices.Add(B);
-				ProjectedTriVertices.Add(C);
-			}
-			double BoundaryArea = 0;
-			if (NumBoundary)
-			{
-				for (int32 PlaneBoundaryVertIdx : Cells.PlaneBoundaries[PlaneIdx])
-				{
-					FVector2d ProjBoundary(PlaneFrames[PlaneIdx].Project(TransformedPlaneBoundaryVertices[PlaneBoundaryVertIdx]));
-					Boundary.Add(ProjBoundary);
-					Bounds2D.Contain(ProjBoundary);
-				}
-				// area check
-				BoundaryArea = 0;
-				for (int32 Idx = 0; Idx+2 < NumBoundary; Idx++)
-				{
-					BoundaryArea += VectorUtil::Area(Boundary[Idx], Boundary[Idx + 1], Boundary[Idx + 2]);
-				}
-				// don't bother triangulating if the whole boundary is a tiny sliver
-				if (BoundaryArea < 1e-3)
-				{
-					return;
-					//continue;
-				}
-				int32 ShouldCollapseEdgeCount = 0;
-				for (int32 Idx = 0, LastIdx = NumBoundary - 1; Idx < NumBoundary; LastIdx = Idx++)
-				{
-					if (Boundary[Idx].DistanceSquared(Boundary[LastIdx]) < 1e-4)
-					{
-						ShouldCollapseEdgeCount++;
-					}
-				}
-				// After collapsing tiny edges, boundary would be a line segment (TODO: this could discard some cases that should not be discarded!  do something a bit smarter)
-				if (NumBoundary - ShouldCollapseEdgeCount < 3)
-				{
-					return;
-					//continue;
-				}
-
-				// Clip planar edges against boundary
-				// This is optional; after Triangulation we filter everything outside the boundary anyway.  TODO: benchmark whether it actually helps performance
-				if (Cells.AssumeConvexCells)
-				{
-					TArray<bool> NukeEdges;
-					NukeEdges.SetNumZeroed(PlanarEdges.Num());
-					for (int32 BoundEdgePrevIdx = Boundary.Num() - 1, BoundEdgeIdx = 0; BoundEdgeIdx < Boundary.Num(); BoundEdgePrevIdx = BoundEdgeIdx++)
-					{
-						const FVector2d& Pt = Boundary[BoundEdgeIdx];
-						FVector2d Dir = Pt - Boundary[BoundEdgePrevIdx];
-						FVector2d EdgeNormal = Dir.Perp();
-						EdgeNormal.Normalize();
-						for (int32 EdgeIdx = 0; EdgeIdx < PlanarEdges.Num(); EdgeIdx++)
-						{
-							if (NukeEdges[EdgeIdx])
-							{
-								continue;
-							}
-							double SDA = (PlanarEdges[EdgeIdx].Key - Pt).Dot(EdgeNormal);
-							double SDB = (PlanarEdges[EdgeIdx].Value - Pt).Dot(EdgeNormal);
-							if (SDB < -PlaneEps && SDA < -PlaneEps )
-							{
-								NukeEdges[EdgeIdx] = true;
-								continue;
-							}
-							if (SDA*SDB < -PlaneEps) //  TODO: use more robust crossing formula?
-							{
-								double T = SDA / (SDA - SDB);
-								FVector2d OnBoundary = PlanarEdges[EdgeIdx].Key * (1 - T) + PlanarEdges[EdgeIdx].Value * T;
-								double SDO = (OnBoundary - Pt).Dot(EdgeNormal);
-								ensure(FMath::Abs(SDO) < 1e-4);
-								if (SDA < 0)
-								{
-									PlanarEdges[EdgeIdx].Key = OnBoundary;
-								}
-								else
-								{
-									PlanarEdges[EdgeIdx].Value = OnBoundary;
-								}
-							}
-						}
-					}
-
-					// copy the non-nuked edges down, and trim the nuked ones
-					int32 RemainingEdgeCount = 0;
-					for (int32 EdgeIdx = 0; EdgeIdx < NukeEdges.Num(); EdgeIdx++)
-					{
-						if (NukeEdges[EdgeIdx])
-						{
-							continue;
-						}
-						if (EdgeIdx != RemainingEdgeCount)
-						{
-							PlanarEdges[RemainingEdgeCount] = PlanarEdges[EdgeIdx];
-						}
-						RemainingEdgeCount++;
-					}
-					PlanarEdges.SetNum(RemainingEdgeCount);
-				}
-
-				// Add boundary edges as planar edges also
-				BoundaryEdgeStart = PlanarEdges.Num();
-				for (int32 BoundIdx = 0; BoundIdx + 1 < Boundary.Num(); BoundIdx++)
-				{
-					PlanarEdges.Emplace(Boundary[BoundIdx], Boundary[BoundIdx + 1]);
-				}
-				PlanarEdges.Emplace(Boundary[Boundary.Num()-1], Boundary[0]);
-			}
-			else // cell has no boundary; set bounding box just on crossing edges
-			{
-				BoundaryEdgeStart = PlanarEdges.Num();
-				for (const TPair<FVector2d, FVector2d>& Edge2d : PlanarEdges)
-				{
-					Bounds2D.Contain(Edge2d.Key);
-					Bounds2D.Contain(Edge2d.Value);
-				}
-			}
-			
-
-			const double ArrangementTol = 1e-4;
-			const double ScaleF = 1.0 / FMathd::Max(.01, Bounds2D.MaxDim());
-			const FVector2d Offset = -Bounds2D.Center();
-			FAxisAlignedBox2d ScaledBounds2D((Bounds2D.Min + Offset)*ScaleF, (Bounds2D.Max + Offset)*ScaleF);
-			FArrangement2d Arrangement(FMath::Max(ScaledBounds2D.MaxDim() / 64, ArrangementTol * 10));
-			Arrangement.VertexSnapTol = ArrangementTol;
-			int32 BoundaryEdgeGroupID = -1;
-			for (int32 EdgeIdx = PlanarEdges.Num()-1; EdgeIdx >= 0; EdgeIdx--)
-			{
-				const TPair<FVector2d, FVector2d>& Edge2d = PlanarEdges[EdgeIdx];
-				int32 EdgeGroupID = EdgeIdx >= BoundaryEdgeStart ? BoundaryEdgeGroupID : EdgeIdx; // give all boundary edges the same group ID
-				Arrangement.Insert((Edge2d.Key+Offset)*ScaleF, (Edge2d.Value+Offset)*ScaleF, EdgeGroupID);
-			}
-			TArray<int32> SkippedEdges;
-			TArray<FIntVector> PlaneTriangulation;
-
-			TArray<int32> NoiseVertexIndices;
-			if (bNoiseOnPlane)
-			{
-				const FNoiseSettings& Noise = InternalMaterials->NoiseSettings.GetValue();
-				const float MinPointSpacing = .1 * float(ScaleF) * AverageGlobalScaleInv;
-				float PointSpacing = FMath::Max(MinPointSpacing, Noise.PointSpacing*float(ScaleF) * AverageGlobalScaleInv);
-
-				// additional safety to avoid asking for hundreds of millions of noise points
-				float Area = (ScaledBounds2D.Max.X - ScaledBounds2D.Min.X) * (ScaledBounds2D.Max.Y - ScaledBounds2D.Min.Y);
-				float ApproxPointsNeeded = Area / (PointSpacing * PointSpacing);
-				float MaxPointsTarget = 100000;
-				if (ApproxPointsNeeded > MaxPointsTarget)
-				{
-					PointSpacing = FMath::Sqrt(Area / MaxPointsTarget);
-					ApproxPointsNeeded = Area / (PointSpacing * PointSpacing);
-				}
-
-				// make a new point hash for blue noise point location queries
-				// this is essentially the same as the point hash in arrangement2d but with cell spacing set based on the point spacing; the arrangement2d one can have a way-too-small point spacing!
-				TPointHashGrid2d<int> NoisePointHash(PointSpacing, -1);
-				auto HasVertexNear = [&NoisePointHash, &Arrangement, PointSpacing](const FVector2d& V, float ScaleFactor = 1.0)
-				{
-					auto FuncDistSq = [&Arrangement, &V](int B) { return V.DistanceSquared(Arrangement.Graph.GetVertex(B)); };
-					TPair<int, double> NearestPt = NoisePointHash.FindNearestInRadius(V, PointSpacing*ScaleFactor, FuncDistSq);
-					return NearestPt.Key != NoisePointHash.GetInvalidValue();
-				};
-				auto AddNoiseVertex = [&NoisePointHash](int32 ID, const FVector2d& Pos)
-				{
-					NoisePointHash.InsertPointUnsafe(ID, Pos);
-				};
-				for (int32 VertIdx = 0; VertIdx < Arrangement.Graph.MaxVertexID(); VertIdx++)
-				{
-					if (Arrangement.Graph.IsVertex(VertIdx))
-					{
-						AddNoiseVertex(VertIdx, Arrangement.Graph.GetVertex(VertIdx));
-					}
-				}
-
-				double SpacingSq = PointSpacing*PointSpacing;
-
-				// split the edges, ensuring there is no span on any edge that is not within Spacing dist of a vertex
-				for (int32 EdgeIdx : Arrangement.Graph.EdgeIndices())
-				{
-					FDynamicGraph::FEdge Edge = Arrangement.Graph.GetEdge(EdgeIdx);
-					FVector2d A, B;
-					
-					Arrangement.Graph.GetEdgeV(EdgeIdx, A, B);
-					FVector2d Diff = B - A;
-					double DSq = Diff.SquaredLength();
-					int32 WantSamples = FMath::Sqrt(DSq / SpacingSq) + 1;
-					if (DSq > SpacingSq)
-					{
-						int EdgeToSplit = EdgeIdx;
-						for (int SampleIdx = 1; SampleIdx < WantSamples; SampleIdx++)
-						{
-							double T = double(SampleIdx) / double(WantSamples);
-							FVector2d Pt = A + Diff * T;
-							if (!HasVertexNear(Pt, .499))
-							{
-								bool bTargetAtEnd = Arrangement.Graph.GetEdge(EdgeToSplit).B == Edge.B;
-								FIndex2i NewVertEdge = Arrangement.SplitEdgeAtPoint(EdgeToSplit, Pt);
-								int NewEdge = NewVertEdge.B;
-								if (bTargetAtEnd)
-								{
-									EdgeToSplit = NewEdge;
-								}
-								
-								AddNoiseVertex(NewVertEdge.A, Pt);
-							}
-						}
-					}
-				}
-				// insert internal noise points
-				for (double X = ScaledBounds2D.Min.X; X < ScaledBounds2D.Max.X; X += PointSpacing)
-				{
-					for (double Y = ScaledBounds2D.Min.Y; Y < ScaledBounds2D.Max.Y; Y += PointSpacing)
-					{
-						for (int Attempt = 0; Attempt < 5; Attempt++)
-						{
-							FVector2d Pt(X + FMath::FRand() * PointSpacing*.5, Y + FMath::FRand() * PointSpacing*.5);
-							if (!HasVertexNear(Pt))
-							{
-								// we shouldn't need to check for overlap with any existing pt/edge b/c of the above edge covering
-								// (and it's much faster not to)
-								int32 PtIdx = Arrangement.InsertNewIsolatedPointUnsafe(Pt);
-								AddNoiseVertex(PtIdx, Pt);
-								NoiseVertexIndices.Add(PtIdx);
-								break;
-							}
-						}
-					}
-				}
-			}
-
-			Arrangement.AttemptTriangulate(PlaneTriangulation, SkippedEdges, BoundaryEdgeGroupID);
-
-			// undo scaling
-			double InvScaleF = 1.0 / ScaleF;
-			for (int GraphVertIdx : Arrangement.Graph.VertexIndices())
-			{
-				Arrangement.Graph.SetVertex(GraphVertIdx, (Arrangement.Graph.GetVertex(GraphVertIdx) * InvScaleF) - Offset);
-			}
-
-			// Eat any triangles that are inside coplanar triangles on the face
-			// TODO: optimize this? currently doing this in the simplest way possible because it is likely a rare case
-			
-			int32 NumCoplanarTris = ProjectedTriVertices.Num() / 3;
-			// note: depending on epsilon, could accept an infinitely-large extent as "on triangle", seems a little dangerous ...
-			auto IsOnTriangle2D = [](const FVector2d& Pt, const TArray<FVector2d>& Tris, int32 TriIdx, double TriSideEps = 1e-6)
-			{
-				int32 IdxStart = TriIdx * 3;
-				
-				char NumSideA = 0;
-				char NumSideB = 0;
-				for (int32 Idx = 0, LastIdx = 2; Idx < 3; LastIdx = Idx++)
-				{
-					FVector2d E = Tris[IdxStart + Idx] - Tris[IdxStart + LastIdx];
-					E.Normalize();
-					double Side = E.DotPerp(Pt - Tris[IdxStart + LastIdx]);
-					if (Side < TriSideEps)
-					{
-						NumSideA++;
-					}
-					if (Side > -TriSideEps)
-					{
-						NumSideB++;
-					}
-				}
-				return NumSideA == 3 || NumSideB == 3;
-			};
-			if (NumCoplanarTris > 0)
-			{
-				int32 Ate = 0;
-				for (int32 PlaneTriIdx = 0, CopyTriIdx = 0; PlaneTriIdx < PlaneTriangulation.Num(); PlaneTriIdx++)
-				{
-					const FIntVector& Tri = PlaneTriangulation[PlaneTriIdx];
-					FVector2d TriCentroid = (Arrangement.Graph.GetVertex(Tri.X) + Arrangement.Graph.GetVertex(Tri.Y) + Arrangement.Graph.GetVertex(Tri.Z)) / 3.0f;
-					bool bEatTri = false;
-					for (int32 CoplanarTriIdx = 0; CoplanarTriIdx < NumCoplanarTris; CoplanarTriIdx++)
-					{
-						if (IsOnTriangle2D(TriCentroid, ProjectedTriVertices, CoplanarTriIdx))
-						{
-							bEatTri = true;
-							Ate++;
-							break;
-						}
-					}
-					if (!bEatTri)
-					{
-						// copy back any triangle that we aren't eating away
-						if (PlaneTriIdx != CopyTriIdx)
-						{
-							PlaneTriangulation[CopyTriIdx] = PlaneTriangulation[PlaneTriIdx];
-						}
-						CopyTriIdx++;
-					}
-				}
-				if (Ate)
-				{
-					PlaneTriangulation.SetNum(PlaneTriangulation.Num() - Ate);
-				}
-			}
-
-			ensure(SkippedEdges.Num() == 0); // TODO: remove this ensure; for now I'm just curious how much triangulation fails in practice
-
-			ensure(Arrangement.Graph.IsCompact());
-			for (int32 VertIdx = 0; VertIdx < Arrangement.Graph.MaxVertexID(); VertIdx++)
-			{
-				Triangulation.LocalVertices.Add((FVector)PlaneFrames[PlaneIdx].UnProject(Arrangement.Graph.GetVertex(VertIdx)));
-			}
-			
-			for (FIntVector& Face : PlaneTriangulation)
-			{
-				FVector3d TriCentroid = (Triangulation.LocalVertices[Face.X] + Triangulation.LocalVertices[Face.Y] + Triangulation.LocalVertices[Face.Z]) / 3.0;
-				float WindingFast = FastWindingTree.FastWindingNumber(TriCentroid);
-				if (WindingFast > .5)
-				{
-					Triangulation.LocalIndices.Add(Face);
-				}
-			}
-			if (Triangulation.LocalIndices.Num() == 0)
-			{
-				Triangulation.LocalVertices.Empty();
-			}
-			else if (bNoiseOnPlane)
-			{
-				float Amplitude = InternalMaterials->NoiseSettings->Amplitude;
-				float Frequency = InternalMaterials->NoiseSettings->Frequency;
-				int32 Octaves = InternalMaterials->NoiseSettings->Octaves;
-				FVector3d Z = PlaneNormal * Amplitude;
-				for (int32 VertexIdx : NoiseVertexIndices)
-				{
-					FVector2D V = (FVector2D)(Arrangement.Graph.GetVertex(VertexIdx) * Frequency);
-					float NoiseValue = 0;
-					float OctaveScale = 1;
-					for (int32 Octave = 0; Octave < Octaves; Octave++, OctaveScale *= 2)
-					{
-						NoiseValue += FMath::PerlinNoise2D(V * OctaveScale * AverageGlobalScale) / OctaveScale;
-					}
-					Triangulation.LocalVertices[VertexIdx] += (FVector)(Z * NoiseValue * AverageGlobalScaleInv);
-				}
-			}
-		}
-		else // no CDT needed; just triangulate the cell directly
-		{
-			ensure(NumBoundary != 1 && NumBoundary != 2);  // Point or Line Segment boundaries would be weird / imply a bug maybe
-			if (NumBoundary > 2) // if there are at least 3 boundary points, we still have something we could triangulate
-			{
-				// TODO: should actually be "AssumeConvexFacets"?
-				ensure(Cells.AssumeConvexCells);  // should have followed the above code path w/ triangulation if non-convex faces
-
-				FVector FacetCentroid(0, 0, 0);
-				for (int32 VertIdx : BoundaryIndices)
-				{
-					FacetCentroid += TransformedPlaneBoundaryVertices[VertIdx];
-				}
-				FacetCentroid /= BoundaryIndices.Num();
-
-				double Winding = FastWindingTree.FastWindingNumber(FacetCentroid);
-				if (Winding > .5)
-				{
-					// create a simple triangle fan covering the convex facet
-					Triangulation.LocalIndices.SetNum(NumBoundary - 2);
-					for (int32 TriIdx = 0; TriIdx + 2 < NumBoundary; TriIdx++)
-					{
-						Triangulation.LocalIndices[TriIdx] = FIntVector(0, TriIdx + 2, TriIdx + 1);
-					}
-					Triangulation.LocalVertices.SetNum(NumBoundary);
-					for (int32 VertIdx = 0; VertIdx < NumBoundary; VertIdx++)
-					{
-						Triangulation.LocalVertices[VertIdx] = TransformedPlaneBoundaryVertices[BoundaryIndices[VertIdx]];
-					}
-				}
-
-			}
-		}
-		// UV projection
-		int32 NumLocalVertices = Triangulation.LocalVertices.Num();
-		if (NumLocalVertices > 0)
-		{
-			Triangulation.LocalUVs.SetNum(NumLocalVertices);
-			FVector FrameX = (FVector)PlaneFrames[PlaneIdx].X;
-			FVector FrameY = (FVector)PlaneFrames[PlaneIdx].Y;
-			FVector LocalOrigin = Triangulation.LocalVertices[0];
-			float MinX = FMathf::MaxReal;
-			float MinY = FMathf::MaxReal;
-			float WorldToUVScaleFactor = InternalMaterials->GlobalUVScale;
-			for (int32 VertIdx = 0; VertIdx < NumLocalVertices; VertIdx++)
-			{
-				FVector VMinusO = Triangulation.LocalVertices[VertIdx] - LocalOrigin;
-				FVector2D ProjectedPt(FVector::DotProduct(FrameX, VMinusO) * WorldToUVScaleFactor, FVector::DotProduct(FrameY, VMinusO) * WorldToUVScaleFactor);
-				MinX = FMath::Min(MinX, ProjectedPt.X);
-				MinY = FMath::Min(MinY, ProjectedPt.Y);
-				Triangulation.LocalUVs[VertIdx] = ProjectedPt;
-			}
-			for (int32 VertIdx = 0; VertIdx < NumLocalVertices; VertIdx++)
-			{
-				Triangulation.LocalUVs[VertIdx].X -= MinX;
-				Triangulation.LocalUVs[VertIdx].Y -= MinY;
-			}
-		}
-	}
-	, bNoParallel);
-
-	// PHASE 4 PART 2: COPY THE PLANAR TRIANGULATION VERTICES INTO THEIR RESPECTIVE CELLS
-
-	// allocate buffers for added vertices
-	int32 TotalAddedVertices = 0;
-	TArray<int32> VertexIndexToGlobalAddedOffset; VertexIndexToGlobalAddedOffset.SetNum(NumPlanes);
-	for (int32 PlaneIdx = 0; PlaneIdx < NumPlanes; PlaneIdx++)
-	{
-		VertexIndexToGlobalAddedOffset[PlaneIdx] = TotalAddedVertices;
-		TotalAddedVertices += PlaneTriangulations[PlaneIdx].LocalVertices.Num() * 2; // add each vertex twice to allow for opposite vertex normals
-	}
-	int32 AddedVertexStart = AddedVerticesCollection.AddElements(TotalAddedVertices, FGeometryCollection::VerticesGroup);
-
-	for (int32 PlaneIdx = 0; PlaneIdx < NumPlanes; PlaneIdx++)
-	{
-		const PlaneTriangulationInfo& Triangulation = PlaneTriangulations[PlaneIdx];
-		if (!Triangulation.LocalIndices.Num())
-		{
-			continue;
-		}
-
-		int32 AddedVertexOffset = VertexIndexToGlobalAddedOffset[PlaneIdx] + AddedVertexStart;
-		FIntVector TriIdxOffset = FIntVector(AddedVertexOffset + InputVertexCount);
-		int32 NumLocalVertices = Triangulation.LocalVertices.Num();
-		FIntVector OtherCellOffset = FIntVector(NumLocalVertices);
-
-		int32 CellA = Output.OutputCellIdx(Cells.PlaneCells[PlaneIdx].Key);
-		int32 CellB = Output.OutputCellIdx(Cells.PlaneCells[PlaneIdx].Value);
-
-		const FPlane &Plane = TransformedPlanes[PlaneIdx];
-		FVector PlaneNormal(Plane.X, Plane.Y, Plane.Z);
-		
-		if (CheckDistanceAcrossOutsideCellForProximity > 0 && (CellA == Output.NoCellIdx || CellB == Output.NoCellIdx))
-		{
-			int32 InsideCell = CellA;
-			FVector Direction = PlaneNormal;
-			if (InsideCell == Output.NoCellIdx)
-			{
-				InsideCell = CellB;
-				Direction = -Direction;
-			}
-			for (const FIntVector &LocalTri : Triangulation.LocalIndices)
-			{
-				FVector C = (
-					Triangulation.LocalVertices[LocalTri.X] +
-					Triangulation.LocalVertices[LocalTri.Y] +
-					Triangulation.LocalVertices[LocalTri.Z]) / 3.0f;
-				int32 AcrossCell = Cells.CellFromPosition(C + Direction * CheckDistanceAcrossOutsideCellForProximity);
-				if (AcrossCell != InsideCell)
-				{
-					Output.ConnectCells(InsideCell, AcrossCell);
-				}
-			}
-		}
-		else
-		{
-			Output.ConnectCells(CellA, CellB);
-		}
-		for (const FIntVector &LocalTri : Triangulation.LocalIndices)
-		{
-			FIntVector GlobalTri = LocalTri + TriIdxOffset;
-			Output.AddTriangle(CellB, -1, GlobalTri + OtherCellOffset);
-			
-			Swap(GlobalTri.Y, GlobalTri.Z);
-			Output.AddTriangle(CellA, -1, GlobalTri);
-		}
-		for (int32 LocalVertIdx = 0; LocalVertIdx < Triangulation.LocalVertices.Num(); LocalVertIdx++)
-		{
-			int32 AddIdx = AddedVertexOffset + LocalVertIdx;
-			AddedVerticesCollection.Vertex[AddIdx] = Triangulation.LocalVertices[LocalVertIdx];
-			AddedVerticesCollection.Vertex[AddIdx + NumLocalVertices] = Triangulation.LocalVertices[LocalVertIdx];
-
-			AddedVerticesCollection.UV[AddIdx] = Triangulation.LocalUVs[LocalVertIdx];
-			AddedVerticesCollection.UV[AddIdx + NumLocalVertices] = Triangulation.LocalUVs[LocalVertIdx];
-
-			AddedVerticesCollection.Normal[AddIdx] = PlaneNormal;
-			AddedVerticesCollection.Normal[AddIdx + NumLocalVertices] = -AddedVerticesCollection.Normal[AddIdx];
-
-			AddedVerticesCollection.TangentU[AddIdx] = (FVector)PlaneFrames[PlaneIdx].X;
-			AddedVerticesCollection.TangentU[AddIdx + NumLocalVertices] = -AddedVerticesCollection.TangentU[AddIdx];
-
-			AddedVerticesCollection.TangentV[AddIdx] = (FVector)PlaneFrames[PlaneIdx].Y;
-			AddedVerticesCollection.TangentV[AddIdx + NumLocalVertices] = AddedVerticesCollection.TangentV[AddIdx];
-
-			// TODO: also set color?
-		}
-	}
-
-
-	// ~~~ PHASE 5: FIGURE OUT VERTEX MAPPING FROM SHARED ORIGINAL DATA INTO NEW CELLS ~~~
-	// Obstacles to parallelize: None?
-	for (int32 CellIdx = 0; CellIdx < Output.Num(); CellIdx++)
-	{
-		TMap<int32, int32> GlobalLocalVertexMap;
-		auto RemapVertex = [](TMap<int32, int32> &AllVertexMap, TArray<int32>& CellVertexMapping, int32 GlobalIdx)
-		{
-			int32 *Mapped = AllVertexMap.Find(GlobalIdx);
-			if (Mapped)
-			{
-				return *Mapped;
-			}
-			int32 LocalIdx = CellVertexMapping.Add(GlobalIdx);
-			AllVertexMap.Add(GlobalIdx, LocalIdx);
-			return LocalIdx;
-		};
-		for (FIntVector& Tri : Output.CellTriangles[CellIdx])
-		{
-			Tri.X = RemapVertex(GlobalLocalVertexMap, Output.CellVertexMapping[CellIdx], Tri.X);
-			Tri.Y = RemapVertex(GlobalLocalVertexMap, Output.CellVertexMapping[CellIdx], Tri.Y);
-			Tri.Z = RemapVertex(GlobalLocalVertexMap, Output.CellVertexMapping[CellIdx], Tri.Z);
-		}
-	}
-}
-
-
-
-void TransformPlanes(const FTransform& Transform, const FPlanarCells& Ref, TArray<FPlane>& Planes, TArray<FVector>& PlaneBoundaries)
-{
-	// Note: custom implementation of normal transform for robustness, especially to ensure we don't zero the normals for significantly scaled geometry
-	FTransform NormalTransform = Transform;
-	FVector3d ScaleVec(NormalTransform.GetScale3D());
-	double ScaleDetSign = FMathd::SignNonZero(ScaleVec.X) * FMathd::SignNonZero(ScaleVec.Y) * FMathd::SignNonZero(ScaleVec.Z);
-	double ScaleMaxAbs = ScaleVec.MaxAbsElement();
-	if (ScaleMaxAbs > DBL_MIN)
-	{
-		ScaleVec /= ScaleMaxAbs;
-	}
-	FVector3d NormalScale(ScaleVec.Y*ScaleVec.Z*ScaleDetSign, ScaleVec.X*ScaleVec.Z*ScaleDetSign, ScaleVec.X*ScaleVec.Y*ScaleDetSign);
-	NormalTransform.SetScale3D((FVector)NormalScale);
-	
-	Planes.SetNum(Ref.Planes.Num());
-	for (int32 PlaneIdx = 0, PlanesNum = Planes.Num(); PlaneIdx < PlanesNum; PlaneIdx++)
-	{
-		FVector Pos = Transform.TransformPosition(Ref.Planes[PlaneIdx] * Ref.Planes[PlaneIdx].W);
-		FVector Normal = NormalTransform.TransformVector(Ref.Planes[PlaneIdx]).GetSafeNormal(FLT_MIN);
-		Planes[PlaneIdx] = FPlane(Pos, Normal);
-	}
-
-	PlaneBoundaries.SetNum(Ref.PlaneBoundaryVertices.Num());
-	for (int32 VertIdx = 0, VertsNum = PlaneBoundaries.Num(); VertIdx < VertsNum; VertIdx++)
-	{
-		PlaneBoundaries[VertIdx] = Transform.TransformPosition(Ref.PlaneBoundaryVertices[VertIdx]);
-	}
-}
-
-
-// Simpler invocation of CutWithPlanarCells w/ reasonable defaults
-int32 CutWithPlanarCells(
-	FPlanarCells& Cells,
-	FGeometryCollection& Source,
-	int32 TransformIdx,
-	double Grout,
-	double CollisionSampleSpacing,
-	const TOptional<FTransform>& TransformCollection,
-	bool bIncludeOutsideCellInOutput,
-	float CheckDistanceAcrossOutsideCellForProximity,
-	bool bSetDefaultInternalMaterialsFromCollection,
-	TFunction<void(const FGeometryCollection&, int32, const FGeometryCollection&, int32, float, int32, FGeometryCollection&)> VertexInterpolate
-)
-{
-	TArray<int32> TransformIndices { TransformIdx };
-	return CutMultipleWithPlanarCells(Cells, Source, TransformIndices, Grout, CollisionSampleSpacing, TransformCollection, bIncludeOutsideCellInOutput, CheckDistanceAcrossOutsideCellForProximity, bSetDefaultInternalMaterialsFromCollection, VertexInterpolate);
-}
-
-int32 CutMultipleWithPlanarCells_MeshBooleanPath(
-	FPlanarCells& Cells,
-	FGeometryCollection& Source,
-	const TArrayView<const int32>& TransformIndices,
-	double Grout,
-	double CollisionSampleSpacing,
-	const TOptional<FTransform>& TransformCollection,
-	bool bIncludeOutsideCellInOutput,
-	float CheckDistanceAcrossOutsideCellForProximity,
-	bool bSetDefaultInternalMaterialsFromCollection,
-	TFunction<void(const FGeometryCollection&, int32, const FGeometryCollection&, int32, float, int32, FGeometryCollection&)> VertexInterpolate
-);
-
-// Cut multiple Geometry groups inside a GeometryCollection with PlanarCells, and add each cut cell back to the GeometryCollection as a new child of their source Geometry
-int32 CutMultipleWithPlanarCells(
-	FPlanarCells &Cells,
-	FGeometryCollection& Source,
-	const TArrayView<const int32>& TransformIndices,
-	double Grout,
-	double CollisionSampleSpacing,
-	const TOptional<FTransform>& TransformCollection,
-	bool bIncludeOutsideCellInOutput,
-	float CheckDistanceAcrossOutsideCellForProximity,
-	bool bSetDefaultInternalMaterialsFromCollection,
-	TFunction<void(const FGeometryCollection&, int32, const FGeometryCollection&, int32, float, int32, FGeometryCollection&)> VertexInterpolate
-)
-{
-	if (CVarFractureMethod.GetValueOnAnyThread() >= 1)
-	{
-		return CutMultipleWithPlanarCells_MeshBooleanPath(
-			Cells,
-			Source,
-			TransformIndices,
-			Grout,
-			CollisionSampleSpacing,
-			TransformCollection,
-			bIncludeOutsideCellInOutput,
-			CheckDistanceAcrossOutsideCellForProximity,
-			bSetDefaultInternalMaterialsFromCollection,
-			VertexInterpolate);
-	}
-	float PlaneEps = 1e-4;
-
-	int32 NewGeomStartIdx = -1;
-
-	FTransform CollectionToWorld = TransformCollection.Get(FTransform::Identity);
-
-	TArray<FPlane> TransformedPlanes;
-	TArray<FVector> TransformedPlaneBoundaries;
-
-	if (!Source.HasAttribute("Proximity", FGeometryCollection::GeometryGroup))
-	{
-		const FManagedArrayCollection::FConstructionParameters GeometryDependency(FGeometryCollection::GeometryGroup);
-		Source.AddAttribute<TSet<int32>>("Proximity", FGeometryCollection::GeometryGroup, GeometryDependency);
-	}
-
-	if (bSetDefaultInternalMaterialsFromCollection)
-	{
-		Cells.InternalSurfaceMaterials.SetUVScaleFromCollection(Source);
-	}
-
-#if WITH_EDITOR
-	// Create progress indicator dialog
-	static const FText SlowTaskText = NSLOCTEXT("CutMultipleWithPlanarCells", "CutMultipleWithPlanarCellsText", "Cutting geometry collection...");
-
-	int32 TotalFacesToProcess = 0;
-	for (int32 TransformIdx : TransformIndices)
-	{
-		int32 GeometryIdx = Source.TransformToGeometryIndex[TransformIdx];
-		TotalFacesToProcess += Source.FaceCount[GeometryIdx];
-	}
-
-	FScopedSlowTask SlowTask(float(TotalFacesToProcess), SlowTaskText);
-	SlowTask.MakeDialog();
-
-	// Declare progress shortcut lambdas
-	auto EnterProgressFrame = [&SlowTask](int32 TaskSize)
-	{
-		SlowTask.EnterProgressFrame(float(TaskSize));
-	};
-#else
-	auto EnterProgressFrame = [](int32 /*TaskSize*/) {};
-#endif
-
-	// TODO: Cannot yet make this a parallel for because the Voronoi pt->cell function is not thread-safe, but after that is fixed, consider threading this?
-	for (int32 ParentTransformIndex : TransformIndices)
-	{
-		int32 GeometryIdx = Source.TransformToGeometryIndex[ParentTransformIndex];
-		EnterProgressFrame(Source.FaceCount[GeometryIdx]);
-		if (Source.Children[ParentTransformIndex].Num())
-		{
-			// don't fracture an already-fractured geometry
-			UE_LOG(LogPlanarCut, Warning, TEXT("Skipping cut of a non-leaf geometry, as this would would create intersecting / duplicate geometry"));
-			continue;
-		}
-		int32 TriangleStart = Source.FaceStart[GeometryIdx];
-		int32 NumTriangles = Source.FaceCount[GeometryIdx];
-		TArray<FVector> TriangleNormals;
-		ComputeTriangleNormals(TArrayView<const FVector>(Source.Vertex.GetData(), Source.Vertex.Num()), TArrayView<const FIntVector>(Source.Indices.GetData() + TriangleStart, NumTriangles), TriangleNormals);
-
-		TArrayView<const FPlane> Planes;
-		TArrayView<const FVector> PlaneBoundaries;
-		FTransform LocalToPlaneSpaceTransform = FTransform::Identity;
-
-		LocalToPlaneSpaceTransform = GeometryCollectionAlgo::GlobalMatrix(Source.Transform, Source.Parent, ParentTransformIndex) * CollectionToWorld;
-		FTransform PlanesToLocalTransform = LocalToPlaneSpaceTransform.Inverse();
-		TransformPlanes(PlanesToLocalTransform, Cells, TransformedPlanes, TransformedPlaneBoundaries);
-		Planes = TransformedPlanes;
-		PlaneBoundaries = TransformedPlaneBoundaries;
-
-		OutputCells Output(Cells.NumCells);
-		// TODO: handle bIncludeOutsideCellInOutput in the CutWithPlanarCellsHelper so you don't even store the no-cell geometry in the temporary Output object (currently we always compute everything and then just don't dump the no-cell to the geom collection if bIncludeOutsideCellInOutput==false)
-		CutWithPlanarCellsHelper(Cells, Planes, PlaneBoundaries, LocalToPlaneSpaceTransform, Source, GeometryIdx, TriangleStart, NumTriangles, TriangleNormals, PlaneEps, CheckDistanceAcrossOutsideCellForProximity, VertexInterpolate, Output);
-		if (Output.NumNonEmptyCells() == 1) // nothing was actually cut; skip cutting this geometry entirely
-		{
-			continue;
-		}
-		int32 SourceVertexNum = Source.Vertex.Num();
-		int MaterialIDOverride = bSetDefaultInternalMaterialsFromCollection ? Cells.InternalSurfaceMaterials.GetDefaultMaterialIDForGeometry(Source, GeometryIdx) : -1;
-		int32 StartIdx = Output.AddToGeometryCollection(Source, Cells.InternalSurfaceMaterials, bIncludeOutsideCellInOutput, SourceVertexNum, ParentTransformIndex, MaterialIDOverride);
-		if (NewGeomStartIdx < 0)
-		{
-			NewGeomStartIdx = StartIdx;
-		}
-
-		// turn off old geom visibility (preferred default behavior)
-		int32 FaceEnd = Source.FaceCount[GeometryIdx] + Source.FaceStart[GeometryIdx];
-		for (int32 FaceIdx = Source.FaceStart[GeometryIdx]; FaceIdx < FaceEnd; FaceIdx++)
-		{
-			Source.Visible[FaceIdx] = false;
-		}
-	}
-
-	// Fix MaterialIndex values
-	Source.ReindexMaterials();
-	return NewGeomStartIdx;
-}
-
-
-
-
-int32 CutMultipleWithMultiplePlanes_old(
-	const TArrayView<const FPlane>& Planes,
-	FInternalSurfaceMaterials& InternalSurfaceMaterials,
-	FGeometryCollection& Collection,
-	const TArrayView<const int32>& TransformIndices,
-	const TOptional<FTransform>& TransformCollection,
-	bool bSetDefaultInternalMaterialsFromCollection,
-	TFunction<void(const FGeometryCollection&, int32, const FGeometryCollection&, int32, float, int32, FGeometryCollection&)> VertexInterpolate
-)
-{
-	bool bFlattenToSingleLayer = true;
-	float PlaneEps = 1e-4;
-
-	int32 OrigNumGeom = Collection.FaceCount.Num();
-	int32 CurNumGeom = OrigNumGeom;
-
-	FTransform CollectionToWorld = TransformCollection.Get(FTransform::Identity);
-
-	TArray<int32> NeedsCut;
-
-	if (bSetDefaultInternalMaterialsFromCollection)
-	{
-		InternalSurfaceMaterials.SetUVScaleFromCollection(Collection);
-	}
-
-	TArray<int32> TransformsToDelete;
-	NeedsCut.SetNum(TransformIndices.Num());
-	for (int32 Idx = 0; Idx < NeedsCut.Num(); Idx++)
-	{
-		// add geometry idx corresponding to each input transform
-		NeedsCut[Idx] = Collection.TransformToGeometryIndex[TransformIndices[Idx]];
-	}
-
-	if (!Collection.HasAttribute("Proximity", FGeometryCollection::GeometryGroup))
-	{
-		const FManagedArrayCollection::FConstructionParameters GeometryDependency(FGeometryCollection::GeometryGroup);
-		Collection.AddAttribute<TSet<int32>>("Proximity", FGeometryCollection::GeometryGroup, GeometryDependency);
-	}
-
-	TMap<int32, TPair<TUniquePtr<FGeometryCollectionMeshAdapter>, TUniquePtr<TMeshAABBTree3<FGeometryCollectionMeshAdapter>>>> AABBTrees;
-	FCriticalSection Mutex;
-	auto GetTree = [&](int32 GeometryIdx)
-	{
-		{
-			FScopeLock Lock(&Mutex);
-			if (AABBTrees.Contains(GeometryIdx))
-			{
-				return AABBTrees[GeometryIdx].Value.Get();
-			}
-		}
-		
-		FGeometryCollectionMeshAdapter* Adapter = new FGeometryCollectionMeshAdapter{ &Collection, GeometryIdx };
-		TMeshAABBTree3<FGeometryCollectionMeshAdapter>* AABBTree = new TMeshAABBTree3<FGeometryCollectionMeshAdapter>(Adapter);
-		{
-			FScopeLock Lock(&Mutex);
-			AABBTrees.Add(GeometryIdx, TPair<TUniquePtr<FGeometryCollectionMeshAdapter>, TUniquePtr<TMeshAABBTree3<FGeometryCollectionMeshAdapter>>>(Adapter, AABBTree));
-		}
-		return AABBTree;
-	};
-
-#if WITH_EDITOR
-	// Create progress indicator dialog
-	static const FText SlowTaskText = NSLOCTEXT("CutMultipleWithMultiplePlanes", "CutMultipleWithMultiplePlanesText", "Cutting geometry collection with plane(s)...");
-
-	const int32 UnitProgressOutOfLoop = 1;  // One progress frame is equivalent to a minimum of 2% of the loop progress
-
-	FScopedSlowTask SlowTask(float(Planes.Num()), SlowTaskText);
-	SlowTask.MakeDialog();
-
-	// Declare progress shortcut lambdas
-	auto EnterProgressFrame = [&SlowTask, UnitProgressOutOfLoop]()
-	{
-		SlowTask.EnterProgressFrame(float(UnitProgressOutOfLoop));
-	};
-#else
-	auto EnterProgressFrame = []() {};
-#endif
-
-	for (int32 PlaneIdx = 0; PlaneIdx < Planes.Num(); PlaneIdx++)
-	{
-		EnterProgressFrame();
-
-		TArray<int32> NeedsDelete, ChildrenOfTheDeleted;
-
-		const FPlane& Plane = Planes[PlaneIdx];
-		int32 LastNumGeom = CurNumGeom;
-		CurNumGeom = Collection.FaceCount.Num();
-
-		// TUniquePtr instead of just storing OutputCells as values because the OutputCells carry GeometryCollections and those will get messed up if the TArray resizes its storage ... TODO: fix better make less bad
-		TArray<TUniquePtr<OutputCells>> AllOutputsForPlane;
-		for (int32 Idx = 0; Idx < NeedsCut.Num(); Idx++)
-		{
-			AllOutputsForPlane.Add(TUniquePtr<OutputCells>(new OutputCells(2)));
-		}
-
-		// helper function to parallelize plane cuts
-		auto CutGeometryWithPlane = [&](int OutputIdx)
-		{
-			int32 GeometryIdx = NeedsCut[OutputIdx];
-
-			int32 ParentTransformIndex = Collection.TransformIndex[GeometryIdx];
-			if (Collection.Children[ParentTransformIndex].Num())
-			{
-				// don't fracture an already-fractured geometry
-				ensureMsgf(false, TEXT("Skipping cut of a non-leaf geometry, as this would would create intersecting / duplicate geometry"));
-				return;
-			}
-
-			FPlane TransformedPlane = Plane;
-			FTransform LocalToPlaneSpaceTransform = GeometryCollectionAlgo::GlobalMatrix(Collection.Transform, Collection.Parent, ParentTransformIndex) * CollectionToWorld;
-			FTransform PlanesToLocalTransform = LocalToPlaneSpaceTransform.Inverse();
-			FMatrix Matrix = PlanesToLocalTransform.ToMatrixWithScale();
-			TransformedPlane = Plane.TransformBy(Matrix);
-
-			if (!FMath::PlaneAABBIntersection(TransformedPlane, Collection.BoundingBox[GeometryIdx]))
-			{
-				// no intersection; can skip
-				return;
-			}
-
-			int32 TriangleStart = Collection.FaceStart[GeometryIdx];
-			int32 NumTriangles = Collection.FaceCount[GeometryIdx];
-			TArray<FVector> TriangleNormals;
-			ComputeTriangleNormals(TArrayView<const FVector>(Collection.Vertex.GetData(), Collection.Vertex.Num()), TArrayView<const FIntVector>(Collection.Indices.GetData() + TriangleStart, NumTriangles), TriangleNormals);
-
-			OutputCells& Output = *AllOutputsForPlane[OutputIdx].Get();
-			FPlanarCells PlaneCells(TransformedPlane);
-			// TODO: handle bIncludeOutsideCellInOutput in the CutWithPlanarCellsHelper so you don't even store the no-cell geometry in the temporary Output object (currently we always compute everything and then just don't dump the no-cell to the geom collection if bIncludeOutsideCellInOutput==false)
-			TMeshAABBTree3<FGeometryCollectionMeshAdapter> *AABBTree = nullptr;
-			if (bFlattenToSingleLayer)
-			{
-				AABBTree = GetTree(GeometryIdx);
-			}
-			CutWithPlanarCellsHelper(PlaneCells, PlaneCells.Planes, PlaneCells.PlaneBoundaryVertices, FTransform::Identity, Collection, GeometryIdx, TriangleStart, NumTriangles, TriangleNormals, PlaneEps, 0, VertexInterpolate, Output, &InternalSurfaceMaterials, AABBTree);
-		};
-
-		// cut all geometries with the given plane
-		bool bNoParallel = false;
-		ParallelFor(NeedsCut.Num(), CutGeometryWithPlane, bNoParallel);
-
-		int32 SourceVertexNumWhenCut = Collection.Vertex.Num();
-
-		// collect outputs
-		for (int32 OutputIdx = 0; OutputIdx < AllOutputsForPlane.Num(); OutputIdx++)
-		{
-			OutputCells& Output = *AllOutputsForPlane[OutputIdx].Get();
-
-			// TODO: handle bIncludeOutsideCellInOutput in the CutWithPlanarCellsHelper so you don't even store the no-cell geometry in the temporary Output object (currently we always compute everything and then just don't dump the no-cell to the geom collection if bIncludeOutsideCellInOutput==false)
-			if (Output.NumNonEmptyCells() <= 1) // nothing was actually cut; skip cutting this geometry entirely
-			{
-				continue;
-			}
-
-			int32 GeometryIdx = NeedsCut[OutputIdx];
-			int32 ParentTransformIndex = Collection.TransformIndex[GeometryIdx];
-			int MaterialIDOverride = bSetDefaultInternalMaterialsFromCollection ? InternalSurfaceMaterials.GetDefaultMaterialIDForGeometry(Collection, GeometryIdx) : -1;
-			int32 AddedStartIdx = Output.AddToGeometryCollection(Collection, InternalSurfaceMaterials, true, SourceVertexNumWhenCut, ParentTransformIndex, MaterialIDOverride);
-			check(Collection.FaceCount.Num() - AddedStartIdx == 2); // plane must have cut two
-
-			// replace the old geometry with the first of the cut halves
-			NeedsCut[OutputIdx] = AddedStartIdx;
-			// add the other half to the end
-			NeedsCut.Add(AddedStartIdx+1);
-
-			// turn off old geom visibility (preferred default behavior)
-			int32 FaceEnd = Collection.FaceCount[GeometryIdx] + Collection.FaceStart[GeometryIdx];
-			for (int32 FaceIdx = Collection.FaceStart[GeometryIdx]; FaceIdx < FaceEnd; FaceIdx++)
-			{
-				Collection.Visible[FaceIdx] = false;
-			}
-
-			if (bFlattenToSingleLayer && GeometryIdx >= OrigNumGeom)
-			{
-				// flag the geometry for deletion, and store where its replacement sub-parts went (for fixing proximity data, below)
-				TransformsToDelete.Add(ParentTransformIndex);
-				NeedsDelete.Add(GeometryIdx);
-				ChildrenOfTheDeleted.Add(AddedStartIdx);
-			}
-		}
-
-		if (bFlattenToSingleLayer)
-		{
-			TArray<FTransform> GlobalTransforms;
-			GeometryCollectionAlgo::GlobalMatrices(Collection.Transform, Collection.Parent, GlobalTransforms);
-
-			// update proximities for the children of the deleted using the proximity data from their to-be-deleted parents
-			for (int32 DeleteIdx = 0; DeleteIdx < NeedsDelete.Num(); DeleteIdx++)
-			{
-				int32 GeometryIdx = NeedsDelete[DeleteIdx];
-				int32 ChildrenLoc = ChildrenOfTheDeleted[DeleteIdx];
-				TManagedArray<TSet<int32>>& Proximity = Collection.GetAttribute<TSet<int32>>("Proximity", FGeometryCollection::GeometryGroup);
-				auto ConnectPair = [&](int32 A, int32 B)
-				{
-					Proximity[A].Add(B);
-					Proximity[B].Add(A);
-				};
-				const double ProximityThresholdDist = 1e-4;
-				for (int32 NbrGeometryIdx : Proximity[GeometryIdx])
-				{
-					int32 NbrDeleteIdx = NeedsDelete.Find(NbrGeometryIdx);
-					if (NbrDeleteIdx != INDEX_NONE)
-					{
-						int32 NbrChildrenLoc = ChildrenOfTheDeleted[NbrDeleteIdx];
-						// we're deleting both geometries, so we have to consider how to connect their children
-						if (DeleteIdx < NbrDeleteIdx)
-						{	// we only consider if DeleteIdx < NbrDeleteIdx because otherwise we can assume that we already handled this relationship when we handled the Nbr earlier
-							for (int NbrChildSubIdx = 0; NbrChildSubIdx < 2; NbrChildSubIdx++)
-							{
-								TMeshAABBTree3<FGeometryCollectionMeshAdapter>* NbrTree = GetTree(NbrChildrenLoc + NbrChildSubIdx);
-								for (int ChildSubIdx = 0; ChildSubIdx < 2; ChildSubIdx++)
-								{
-									TMeshAABBTree3<FGeometryCollectionMeshAdapter>* ChildTree = GetTree(ChildrenLoc + ChildSubIdx);
-									int32 ChildTransformIdx = Collection.TransformIndex[ChildrenLoc + ChildSubIdx];
-									int32 NbrTransformIdx = Collection.TransformIndex[NbrChildrenLoc + NbrChildSubIdx];
-									FTransform NbrToLocalTransform = GlobalTransforms[NbrTransformIdx] * GlobalTransforms[ChildTransformIdx].Inverse();
-									auto NbrToLocal = [&](const FVector3d& V)
-									{
-										return FVector3d(NbrToLocalTransform.TransformPosition((FVector)V));
-									};
-									double OutDist;
-									ChildTree->FindNearestTriangles(*NbrTree, NbrToLocal, OutDist, ProximityThresholdDist);
-									if (OutDist < ProximityThresholdDist)
-									{
-										ConnectPair(NbrChildrenLoc + NbrChildSubIdx, ChildrenLoc + ChildSubIdx);
-									}
-								}
-							}
-							
-						}
-					}
-					else
-					{
-						// we're just deleting this geometry; can connect directly to the neighbor
-						FMatrix WorldToNbrGeom = GlobalTransforms[NbrGeometryIdx].Inverse().ToMatrixWithScale();
-						int32 PlaneSide = FMath::PlaneAABBRelativePosition(Plane.TransformBy(WorldToNbrGeom), Collection.BoundingBox[NbrGeometryIdx]);
-
-						switch (PlaneSide)
-						{
-						case -1:
-							ConnectPair(ChildrenLoc, NbrGeometryIdx);
-							break;
-						case 1:
-							ConnectPair(ChildrenLoc + 1, NbrGeometryIdx);
-							break;
-						case 0:
-						{
-							TMeshAABBTree3<FGeometryCollectionMeshAdapter> *NbrTree = GetTree(NbrGeometryIdx);
-							for (int ChildSubIdx = 0; ChildSubIdx < 2; ChildSubIdx++)
-							{
-								TMeshAABBTree3<FGeometryCollectionMeshAdapter> *ChildTree = GetTree(ChildrenLoc + ChildSubIdx);
-								int32 ChildTransformIdx = Collection.TransformIndex[ChildrenLoc + ChildSubIdx];
-								int32 NbrTransformIdx = Collection.TransformIndex[NbrGeometryIdx];
-								FTransform NbrToLocalTransform = GlobalTransforms[NbrTransformIdx] * GlobalTransforms[ChildTransformIdx].Inverse();
-								auto NbrToLocal = [&](const FVector3d& V)
-								{
-									return FVector3d(NbrToLocalTransform.TransformPosition((FVector)V));
-								};
-								double OutDist;
-								ChildTree->FindNearestTriangles(*NbrTree, NbrToLocal, OutDist, ProximityThresholdDist);
-								if (OutDist < ProximityThresholdDist)
-								{
-									ConnectPair(NbrGeometryIdx, ChildrenLoc + ChildSubIdx);
-								}
-							}
-						}
-						break;
-						default:
-							ensure(false); // PlaneSide must be -1, 0, or 1
-						}
-
-					}
-				}
-			}
-			for (int32 DeleteIdx = 0; DeleteIdx < NeedsDelete.Num(); DeleteIdx++)
-			{
-				AABBTrees.Remove(NeedsDelete[DeleteIdx]);
-			}
-		}
-	}
-
-	if (bFlattenToSingleLayer)
-	{
-		TransformsToDelete.Sort();
-		Collection.RemoveElements(FGeometryCollection::TransformGroup, TransformsToDelete);
-	}
-
-	// Fix MaterialIndex values
-	Collection.ReindexMaterials();
-
-	return OrigNumGeom == Collection.FaceCount.Num() ? -1 : OrigNumGeom;
-}
-
-
-int32 CutMultipleWithMultiplePlanes(
-	const TArrayView<const FPlane>& Planes,
-	FInternalSurfaceMaterials& InternalSurfaceMaterials,
-	FGeometryCollection& Collection,
-	const TArrayView<const int32>& TransformIndices,
-	double Grout,
-	double CollisionSampleSpacing,
-	const TOptional<FTransform>& TransformCollection,
-	bool bSetDefaultInternalMaterialsFromCollection,
-	TFunction<void(const FGeometryCollection&, int32, const FGeometryCollection&, int32, float, int32, FGeometryCollection&)> VertexInterpolate
-)
-{
-	if (CVarFractureMethod.GetValueOnAnyThread() <= 0)
-	{
-		return CutMultipleWithMultiplePlanes_old(Planes, InternalSurfaceMaterials, Collection, TransformIndices, TransformCollection, bSetDefaultInternalMaterialsFromCollection, VertexInterpolate);
-	}
-	int32 OrigNumGeom = Collection.FaceCount.Num();
-	int32 CurNumGeom = OrigNumGeom;
-
-	if (bSetDefaultInternalMaterialsFromCollection)
-	{
-		InternalSurfaceMaterials.SetUVScaleFromCollection(Collection);
-	}
-
-	if (!Collection.HasAttribute("Proximity", FGeometryCollection::GeometryGroup))
-	{
-		const FManagedArrayCollection::FConstructionParameters GeometryDependency(FGeometryCollection::GeometryGroup);
-		Collection.AddAttribute<TSet<int32>>("Proximity", FGeometryCollection::GeometryGroup, GeometryDependency);
-	}
-
-	FTransform CollectionToWorld = TransformCollection.Get(FTransform::Identity);
-
-	FDynamicMeshCollection MeshCollection(&Collection, TransformIndices, CollectionToWorld);
-	double OnePercentExtend = MeshCollection.Bounds.MaxDim() * .01;
-
-	int32 NewGeomStartIdx = -1;
-	NewGeomStartIdx = MeshCollection.CutWithMultiplePlanes(Planes, Grout, CollisionSampleSpacing, &Collection, InternalSurfaceMaterials, bSetDefaultInternalMaterialsFromCollection);
-
-	Collection.ReindexMaterials();
-	return NewGeomStartIdx;
-}
-
-
-// Cut multiple Geometry groups inside a GeometryCollection with PlanarCells, and add each cut cell back to the GeometryCollection as a new child of their source Geometry
-int32 CutMultipleWithPlanarCells_MeshBooleanPath(
-	FPlanarCells& Cells,
-	FGeometryCollection& Source,
-	const TArrayView<const int32>& TransformIndices,
-	double Grout,
-	double CollisionSampleSpacing,
-	const TOptional<FTransform>& TransformCollection,
-	bool bIncludeOutsideCellInOutput,
-	float CheckDistanceAcrossOutsideCellForProximity,
-	bool bSetDefaultInternalMaterialsFromCollection,
-	TFunction<void(const FGeometryCollection&, int32, const FGeometryCollection&, int32, float, int32, FGeometryCollection&)> VertexInterpolate
-)
-{
-	if (!Source.HasAttribute("Proximity", FGeometryCollection::GeometryGroup))
-	{
-		const FManagedArrayCollection::FConstructionParameters GeometryDependency(FGeometryCollection::GeometryGroup);
-		Source.AddAttribute<TSet<int32>>("Proximity", FGeometryCollection::GeometryGroup, GeometryDependency);
-	}
-
-	if (bSetDefaultInternalMaterialsFromCollection)
-	{
-		Cells.InternalSurfaceMaterials.SetUVScaleFromCollection(Source);
-	}
-
-	FTransform CollectionToWorld = TransformCollection.Get(FTransform::Identity);
-
-	FDynamicMeshCollection MeshCollection(&Source, TransformIndices, CollectionToWorld);
-	double OnePercentExtend = MeshCollection.Bounds.MaxDim() * .01;
-	FCellMeshes CellMeshes(Cells, MeshCollection.Bounds, Grout, OnePercentExtend, bIncludeOutsideCellInOutput);
-
-	int32 NewGeomStartIdx = -1;
-
-	NewGeomStartIdx = MeshCollection.CutWithCellMeshes(Cells.InternalSurfaceMaterials, Cells.PlaneCells, CellMeshes, &Source, bSetDefaultInternalMaterialsFromCollection, CollisionSampleSpacing);
+	NewGeomStartIdx = MeshCollection.CutWithCellMeshes(Cells.InternalSurfaceMaterials, Cells.PlaneCells, CellMeshes, &Source, bSetDefaultInternalMaterialsFromCollection, CollisionSampleSpacing);
 
 	Source.ReindexMaterials();
 	return NewGeomStartIdx;
 }
 
 
-
-
 int32 CutWithMesh(
 	FMeshDescription* CuttingMesh,
 	FTransform CuttingMeshTransform,
diff --git a/Engine/Plugins/Experimental/PlanarCutPlugin/Source/PlanarCut/Public/PlanarCut.h b/Engine/Plugins/Experimental/PlanarCutPlugin/Source/PlanarCut/Public/PlanarCut.h
--- a/Engine/Plugins/Experimental/PlanarCutPlugin/Source/PlanarCut/Public/PlanarCut.h
+++ b/Engine/Plugins/Experimental/PlanarCutPlugin/Source/PlanarCut/Public/PlanarCut.h
@@ -57,9 +57,6 @@ struct PLANARCUT_API FPlanarCells
 	TArray<TArray<int32>> PlaneBoundaries;
 	TArray<FVector> PlaneBoundaryVertices;
 
-	// most cellular complexes we're interested in lend themselves to reasonably-fast ways to go directly from a point in space to a cell classification, so we don't try to compute it directly from this generic representation
-	TFunction<int32(FVector)> CellFromPosition;
-
 	FInternalSurfaceMaterials InternalSurfaceMaterials;
 
 	/**
@@ -117,14 +114,6 @@ struct PLANARCUT_API FPlanarCells
 		return true;
 	}
 
-	void EmptyGeometricData()
-	{
-		PlaneCells.Empty();
-		PlaneBoundaries.Empty();
-		PlaneBoundaryVertices.Empty();
-		Planes.Empty();
-	}
-
 	inline void AddPlane(const FPlane &P, int32 CellIdxBehind, int32 CellIdxInFront)
 	{
 		Planes.Add(P);
@@ -145,11 +134,6 @@ struct PLANARCUT_API FPlanarCells
 	}
 };
 
-// helper function that interpolates the standard vertex attributes in a reasonable way
-void PLANARCUT_API DefaultVertexInterpolation(const FGeometryCollection& V0Collection, int32 V0, const FGeometryCollection& V1Collection, int32 V1, float T, int32 VOut, FGeometryCollection& Dest);
-
-// TODO: this functionality shouldn't live in this api probably; get it from geometry processing modules or something else?
-void PLANARCUT_API ComputeTriangleNormals(const TArrayView<const FVector> Vertices, const TArrayView<const FIntVector> Triangles, TArray<FVector>& TriangleNormals);
 
 /**
  * Cut a Geometry inside a GeometryCollection with PlanarCells, and add each cut cell back to the GeometryCollection as a new child of the input Geometry.  For geometries that would not be cut, nothing is added.
@@ -162,7 +146,6 @@ void PLANARCUT_API ComputeTriangleNormals(const TArrayView<const FVector> Vertic
  * @param TransformCollection		Optional transform of the whole geometry collection; if unset, defaults to Identity
  * @param bIncludeOutsideCellInOutput	If true, geometry that was not inside any of the cells (e.g. was outside of the bounds of all cutting geometry) will still be included in the output; if false, it will be discarded.
  * @param CheckDistanceAcrossOutsideCellForProximity	If > 0, when a plane is neighboring the "outside" cell, instead of setting proximity to the outside cell, the algo will sample a point this far outside the cell in the normal direction of the plane to see if there is actually a non-outside cell there.  (Useful for bricks w/out mortar)
- * @param VertexInterpolate	Function that interpolates vertex properties (UVs, normals, etc); a default that handles all the normal vertex properties is provided, should only need to replace this if you have custom attributes
  * @return	index of first new geometry in the Output GeometryCollection, or -1 if no geometry was added
  */
 int32 PLANARCUT_API CutWithPlanarCells(
@@ -174,8 +157,7 @@ int32 PLANARCUT_API CutWithPlanarCells(
 	const TOptional<FTransform>& TransformCollection = TOptional<FTransform>(),
 	bool bIncludeOutsideCellInOutput = true,
 	float CheckDistanceAcrossOutsideCellForProximity = 0,
-	bool bSetDefaultInternalMaterialsFromCollection = true,
-	TFunction<void(const FGeometryCollection&, int32, const FGeometryCollection&, int32, float, int32, FGeometryCollection&)> VertexInterpolate = DefaultVertexInterpolation
+	bool bSetDefaultInternalMaterialsFromCollection = true
 );
 
 /**
@@ -189,7 +171,6 @@ int32 PLANARCUT_API CutWithPlanarCells(
  * @param TransformCollection		Optional transform of the whole geometry collection; if unset, defaults to Identity
  * @param bIncludeOutsideCellInOutput	If true, geometry that was not inside any of the cells (e.g. was outside of the bounds of all cutting geometry) will still be included in the output; if false, it will be discarded.
  * @param CheckDistanceAcrossOutsideCellForProximity	If > 0, when a plane is neighboring the "outside" cell, instead of setting proximity to the outside cell, the algo will sample a point this far outside the cell in the normal direction of the plane to see if there is actually a non-outside cell there.  (Useful for bricks w/out mortar)
- * @param VertexInterpolate	Function that interpolates vertex properties (UVs, normals, etc); a default that handles all the normal vertex properties is provided, should only need to replace this if you have custom attributes
  * @return	index of first new geometry in the Output GeometryCollection, or -1 if no geometry was added
  */
 int32 PLANARCUT_API CutMultipleWithPlanarCells(
@@ -201,8 +182,7 @@ int32 PLANARCUT_API CutMultipleWithPlanarCells(
 	const TOptional<FTransform>& TransformCollection = TOptional<FTransform>(),
 	bool bIncludeOutsideCellInOutput = true,
 	float CheckDistanceAcrossOutsideCellForProximity = 0,  // TODO: < this param does nothing in the new mode; is only needed in special cases that aren't possible in the UI currently
-	bool bSetDefaultInternalMaterialsFromCollection = true,
-	TFunction<void(const FGeometryCollection&, int32, const FGeometryCollection&, int32, float, int32, FGeometryCollection&)> VertexInterpolate = DefaultVertexInterpolation
+	bool bSetDefaultInternalMaterialsFromCollection = true
 );
 
 /**
@@ -215,7 +195,6 @@ int32 PLANARCUT_API CutMultipleWithPlanarCells(
  * @param Grout				Separation to leave between cutting cells
  * @param CollisionSampleSpacing	Target spacing between collision sample vertices
  * @param TransformCollection		Optional transform of the whole geometry collection; if unset, defaults to Identity
- * @param VertexInterpolate	Function that interpolates vertex properties (UVs, normals, etc); a default that handles all the normal vertex properties is provided, should only need to replace this if you have custom attributes
  * @return	index of first new geometry in the Output GeometryCollection, or -1 if no geometry was added
  */
 int32 PLANARCUT_API CutMultipleWithMultiplePlanes(
@@ -226,8 +205,7 @@ int32 PLANARCUT_API CutMultipleWithMultiplePlanes(
 	double Grout,
 	double CollisionSampleSpacing,
 	const TOptional<FTransform>& TransformCollection = TOptional<FTransform>(),
-	bool bSetDefaultInternalMaterialsFromCollection = true,
-	TFunction<void(const FGeometryCollection&, int32, const FGeometryCollection&, int32, float, int32, FGeometryCollection&)> VertexInterpolate = DefaultVertexInterpolation
+	bool bSetDefaultInternalMaterialsFromCollection = true
 );
 
 
