remove unused code and fix bug where a bounding box was not expanded as intended

#rnx

diff --git a/Engine/Plugins/Experimental/PlanarCutPlugin/Source/PlanarCut/Private/PlanarCut.cpp b/Engine/Plugins/Experimental/PlanarCutPlugin/Source/PlanarCut/Private/PlanarCut.cpp
--- a/Engine/Plugins/Experimental/PlanarCutPlugin/Source/PlanarCut/Private/PlanarCut.cpp
+++ b/Engine/Plugins/Experimental/PlanarCutPlugin/Source/PlanarCut/Private/PlanarCut.cpp
@@ -40,66 +40,10 @@
 #endif
 
 
-/**
- * Adapter that lets the generic fast winding and AABB tree code view a geometry from a geometrycollection as a simple index buffer triangle mesh
- */
-struct FGeometryCollectionMeshAdapter
-{
-	const FGeometryCollection *Collection;
-	const int32 GeometryIdx;
 
 
-	constexpr inline bool IsTriangle(int32 index) const
-	{
-		return true;
-	}
-	constexpr inline bool IsVertex(int32 index) const
-	{
-		return true;
-	}
-	inline int32 MaxTriangleID() const
-	{
-		return Collection->FaceCount[GeometryIdx];
-	}
-	inline int32 MaxVertexID() const
-	{
-		return Collection->VertexCount[GeometryIdx];
-	}
-	inline int32 TriangleCount() const
-	{
-		return Collection->FaceCount[GeometryIdx];
-	}
-	inline int32 VertexCount() const
-	{
-		return Collection->VertexCount[GeometryIdx];
-	}
-	constexpr inline int32 GetShapeTimestamp() const
-	{
-		return 0;
-	}
-	inline FIndex3i GetTriangle(int32 Idx) const
-	{
-		int32 VertexStart = Collection->VertexStart[GeometryIdx];
-		FIndex3i Tri(Collection->Indices[Idx + Collection->FaceStart[GeometryIdx]]);
-		Tri.A -= VertexStart;
-		Tri.B -= VertexStart;
-		Tri.C -= VertexStart;
-		return Tri;
-	}
-	inline FVector3d GetVertex(int32 Idx) const
-	{
-		return FVector3d(Collection->Vertex[Idx + Collection->VertexStart[GeometryIdx]]);
-	}
 
-	inline void GetTriVertices(int TID, FVector3d& V0, FVector3d& V1, FVector3d& V2) const
-	{
-		FIntVector TriRaw = Collection->Indices[TID + Collection->FaceStart[GeometryIdx]];
 
-		V0 = Collection->Vertex[TriRaw.X];
-		V1 = Collection->Vertex[TriRaw.Y];
-		V2 = Collection->Vertex[TriRaw.Z];
-	}
-};
 
 namespace UE
 {
@@ -1514,7 +1458,8 @@ struct FDynamicMeshCollection
 			EnterProgressFrame(1);
 			FPlanarCells PlaneCells(Planes[PlaneIdx]);
 			PlaneCells.InternalSurfaceMaterials = InternalSurfaceMaterials;
-			FCellMeshes CellMeshes(PlaneCells, Bounds, 0, 0, false);
+			double OnePercentExtend = Bounds.MaxDim() * .01;
+			FCellMeshes CellMeshes(PlaneCells, Bounds, 0, OnePercentExtend, false);
 
 			// TODO: we could do these cuts in parallel (will takes some rework of the proximity and how results are added to the ToCut array)
 			for (int32 ToCutIdx = 0, ToCutNum = ToCut.Num(); ToCutIdx < ToCutNum; ToCutIdx++)
@@ -2623,7 +2568,6 @@ int32 CutMultipleWithMultiplePlanes(
 	FTransform CollectionToWorld = TransformCollection.Get(FTransform::Identity);
 
 	FDynamicMeshCollection MeshCollection(&Collection, TransformIndices, CollectionToWorld);
-	double OnePercentExtend = MeshCollection.Bounds.MaxDim() * .01;
 
 	int32 NewGeomStartIdx = -1;
 	NewGeomStartIdx = MeshCollection.CutWithMultiplePlanes(Planes, Grout, CollisionSampleSpacing, &Collection, InternalSurfaceMaterials, bSetDefaultInternalMaterialsFromCollection);
