Add initial version of an "AutoUV" tool to fracture
This first version lays out non-overlapping UV islands for internal surfaces and generates a texture w/ red channel proportional to distance to the outer surface, and blue channel just indicating if the texel corresponds to an internal surface or not.  Future plan would be to add additional options for e.g. a 'thickness' texture, or whatever else is useful.

diff --git a/Engine/Plugins/Experimental/ChaosEditor/ChaosEditor.uplugin b/Engine/Plugins/Experimental/ChaosEditor/ChaosEditor.uplugin
--- a/Engine/Plugins/Experimental/ChaosEditor/ChaosEditor.uplugin
+++ b/Engine/Plugins/Experimental/ChaosEditor/ChaosEditor.uplugin
@@ -29,6 +29,10 @@
 		{
 			"Name": "GeometryCollectionPlugin",
 			"Enabled": true
+		},
+		{
+			"Name": "MeshModelingToolset",
+			"Enabled": true
 		}
 	]
 }
\ No newline at end of file
diff --git a/Engine/Plugins/Experimental/ChaosEditor/Source/FractureEditor/FractureEditor.Build.cs b/Engine/Plugins/Experimental/ChaosEditor/Source/FractureEditor/FractureEditor.Build.cs
--- a/Engine/Plugins/Experimental/ChaosEditor/Source/FractureEditor/FractureEditor.Build.cs
+++ b/Engine/Plugins/Experimental/ChaosEditor/Source/FractureEditor/FractureEditor.Build.cs
@@ -47,6 +47,7 @@ public class FractureEditor : ModuleRules
 				"LevelEditor",
                 "GeometryCollectionEngine",
                 "GeometryCollectionEditor",
+				"ModelingComponents",
 				"PlanarCut",
 				"Chaos"
 
diff --git a/Engine/Plugins/Experimental/ChaosEditor/Source/FractureEditor/Private/FractureEditorCommands.h b/Engine/Plugins/Experimental/ChaosEditor/Source/FractureEditor/Private/FractureEditorCommands.h
--- a/Engine/Plugins/Experimental/ChaosEditor/Source/FractureEditor/Private/FractureEditorCommands.h
+++ b/Engine/Plugins/Experimental/ChaosEditor/Source/FractureEditor/Private/FractureEditorCommands.h
@@ -45,6 +45,9 @@ class FFractureEditorCommands : public TCommands<FFractureEditorCommands>
 		// Embed Commands
 		TSharedPtr< FUICommandInfo > AddEmbeddedGeometry;
 		TSharedPtr< FUICommandInfo > AutoEmbedGeometry;
+
+		// UV Commands
+		TSharedPtr< FUICommandInfo > AutoUV;
 		
 		// Fracture Commands
 		TSharedPtr< FUICommandInfo > Uniform;
diff --git a/Engine/Plugins/Experimental/ChaosEditor/Source/FractureEditor/Private/FractureEditorModeToolkit.cpp b/Engine/Plugins/Experimental/ChaosEditor/Source/FractureEditor/Private/FractureEditorModeToolkit.cpp
--- a/Engine/Plugins/Experimental/ChaosEditor/Source/FractureEditor/Private/FractureEditorModeToolkit.cpp
+++ b/Engine/Plugins/Experimental/ChaosEditor/Source/FractureEditor/Private/FractureEditorModeToolkit.cpp
@@ -463,7 +463,7 @@ void FFractureEditorModeToolkit::OnObjectPostEditChange( UObject* Object, FPrope
 	}
 }
 
-const TArray<FName> FFractureEditorModeToolkit::PaletteNames = { FName(TEXT("Generate")), FName(TEXT("Select")), FName(TEXT("Fracture")), FName(TEXT("Edit")), FName(TEXT("Cluster")), FName(TEXT("Embed")), FName(TEXT("Properties")) };
+const TArray<FName> FFractureEditorModeToolkit::PaletteNames = { FName(TEXT("Generate")), FName(TEXT("Select")), FName(TEXT("Fracture")), FName(TEXT("Edit")), FName(TEXT("Cluster")), FName(TEXT("Embed")), FName(TEXT("UV")), FName(TEXT("Properties")) };
 
 FText FFractureEditorModeToolkit::GetToolPaletteDisplayName(FName Palette) const
 { 
@@ -530,6 +530,10 @@ void FFractureEditorModeToolkit::BuildToolPalette(FName PaletteIndex, class FToo
 		ToolbarBuilder.AddToolBarButton(Commands.AddEmbeddedGeometry);
 		ToolbarBuilder.AddToolBarButton(Commands.AutoEmbedGeometry);
 	}
+	else if (PaletteIndex == TEXT("UV"))
+	{
+		ToolbarBuilder.AddToolBarButton(Commands.AutoUV);
+	}
 	else if (PaletteIndex == TEXT("Properties"))
 	{
 		ToolbarBuilder.AddToolBarButton(Commands.SetInitialDynamicState);
diff --git a/Engine/Plugins/Experimental/ChaosEditor/Source/FractureEditor/Private/FractureToolResample.cpp b/Engine/Plugins/Experimental/ChaosEditor/Source/FractureEditor/Private/FractureToolResample.cpp
--- a/Engine/Plugins/Experimental/ChaosEditor/Source/FractureEditor/Private/FractureToolResample.cpp
+++ b/Engine/Plugins/Experimental/ChaosEditor/Source/FractureEditor/Private/FractureToolResample.cpp
@@ -29,7 +29,7 @@ FText UFractureToolResample::GetTooltipText() const
 
 FSlateIcon UFractureToolResample::GetToolIcon() const 
 {
-	return FSlateIcon("FractureEditorStyle", "FractureEditor.Mesh");
+	return FSlateIcon("FractureEditorStyle", "FractureEditor.Resample");
 }
 
 void UFractureToolResample::RegisterUICommand( FFractureEditorCommands* BindingContext ) 
diff --git a/Engine/Plugins/Experimental/ChaosEditor/Source/FractureEditor/Private/FractureToolUV.cpp b/Engine/Plugins/Experimental/ChaosEditor/Source/FractureEditor/Private/FractureToolUV.cpp
new file mode 100644
--- /dev/null
+++ b/Engine/Plugins/Experimental/ChaosEditor/Source/FractureEditor/Private/FractureToolUV.cpp
@@ -0,0 +1,241 @@
+// Copyright Epic Games, Inc. All Rights Reserved.
+
+#include "FractureToolUV.h"
+
+#include "FractureToolContext.h"
+
+#include "GeometryCollection/GeometryCollectionObject.h"
+#include "GeometryCollection/GeometryCollectionComponent.h"
+#include "GeometryCollection/GeometryCollectionAlgo.h"
+
+#include "AssetUtils/Texture2DBuilder.h"
+#include "AssetToolsModule.h"
+#include "AssetRegistryModule.h"
+// for content-browser things
+#include "ContentBrowserModule.h"
+#include "IContentBrowserSingleton.h"
+
+#include "Misc/ScopedSlowTask.h"
+
+#include "FractureAutoUV.h"
+
+#define LOCTEXT_NAMESPACE "FractureToolAutoUV"
+
+
+UFractureToolAutoUV::UFractureToolAutoUV(const FObjectInitializer& ObjInit)
+	: Super(ObjInit)
+{
+	AutoUVSettings = NewObject<UFractureAutoUVSettings>(GetTransientPackage(), UFractureAutoUVSettings::StaticClass());
+	AutoUVSettings->OwnerTool = this;
+}
+
+bool UFractureToolAutoUV::CanExecute() const
+{
+	if (!IsGeometryCollectionSelected())
+	{
+		return false;
+	}
+
+	return true;
+}
+
+FText UFractureToolAutoUV::GetDisplayText() const
+{
+	return FText(LOCTEXT("FractureToolAutoUV", "AutoUV Fracture"));
+}
+
+FText UFractureToolAutoUV::GetTooltipText() const
+{
+	return FText(LOCTEXT("FractureToolAutoUVTooltip", "This enables you to automatically layout UVs for internal fracture pieces, and procedurally fill a corresponding texture."));
+}
+
+FSlateIcon UFractureToolAutoUV::GetToolIcon() const
+{
+	return FSlateIcon("FractureEditorStyle", "FractureEditor.AutoUV");
+}
+
+void UFractureToolAutoUV::RegisterUICommand(FFractureEditorCommands* BindingContext)
+{
+	UI_COMMAND_EXT(BindingContext, UICommandInfo, "AutoUV", "AutoUV", "Autogenerate UV and texture for internal faces", EUserInterfaceActionType::ToggleButton, FInputChord());
+	BindingContext->AutoUV = UICommandInfo;
+}
+
+TArray<UObject*> UFractureToolAutoUV::GetSettingsObjects() const
+{
+	TArray<UObject*> Settings;
+	Settings.Add(AutoUVSettings);
+	return Settings;
+}
+
+void UFractureToolAutoUV::FractureContextChanged()
+{
+}
+
+void UFractureToolAutoUV::Render(const FSceneView* View, FViewport* Viewport, FPrimitiveDrawInterface* PDI)
+{
+}
+
+void UFractureToolAutoUV::PostEditChangeChainProperty(struct FPropertyChangedChainEvent& PropertyChangedEvent)
+{
+	// update any cached data 
+}
+
+TArray<FFractureToolContext> UFractureToolAutoUV::GetFractureToolContexts() const
+{
+	TArray<FFractureToolContext> Contexts;
+
+	// A context is gathered for each selected GeometryCollection component, or for each individual bone if Group Fracture is not used.
+	TSet<UGeometryCollectionComponent*> GeomCompSelection;
+	GetSelectedGeometryCollectionComponents(GeomCompSelection);
+
+	for (UGeometryCollectionComponent* GeometryCollectionComponent : GeomCompSelection)
+	{
+		// Generate a context for each selected node
+		FFractureToolContext FullSelection(GeometryCollectionComponent);
+		FullSelection.ConvertSelectionToRigidNodes();
+		
+		// Update global transforms and bounds -- TODO: pull this bounds update out to a shared function?
+		const TManagedArray<FTransform>& Transform = FullSelection.GetGeometryCollection()->Transform;
+		const TManagedArray<int32>& TransformToGeometryIndex = FullSelection.GetGeometryCollection()->TransformToGeometryIndex;
+		const TManagedArray<FBox>& BoundingBoxes = FullSelection.GetGeometryCollection()->BoundingBox;
+
+		TArray<FTransform> Transforms;
+		GeometryCollectionAlgo::GlobalMatrices(Transform, FullSelection.GetGeometryCollection()->Parent, Transforms);
+
+		FBox Bounds(ForceInit);
+		for (int32 BoneIndex : FullSelection.GetSelection())
+		{
+			int32 GeometryIndex = TransformToGeometryIndex[BoneIndex];
+			if (GeometryIndex > INDEX_NONE)
+			{
+				FBox BoneBound = BoundingBoxes[GeometryIndex].TransformBy(Transforms[BoneIndex]);
+				Bounds += BoneBound;
+			}
+		}
+		FullSelection.SetBounds(Bounds);
+
+		Contexts.Add(FullSelection);
+	}
+
+	return Contexts;
+}
+
+bool UFractureToolAutoUV::SaveGeneratedTexture(UTexture2D* GeneratedTexture, FString ObjectBaseName, const UObject* RelativeToAsset, bool bPromptToSave)
+{
+	check(RelativeToAsset);
+	check(GeneratedTexture);
+	check(GeneratedTexture->GetOuter() == GetTransientPackage());
+	check(GeneratedTexture->Source.IsValid());	// texture needs to have valid source data to be savd
+
+	// find path to reference asset
+	UPackage* AssetOuterPackage = CastChecked<UPackage>(RelativeToAsset->GetOuter());
+	FString AssetPackageName = AssetOuterPackage->GetName();
+	FString PackageFolderPath = FPackageName::GetLongPackagePath(AssetPackageName);
+
+	// Show the modal dialog and then get the path/name.
+	// If the user cancels, texture is left as transient
+	if (bPromptToSave)
+	{
+		IContentBrowserSingleton& ContentBrowser = FModuleManager::LoadModuleChecked<FContentBrowserModule>("ContentBrowser").Get();
+
+		FString UseDefaultAssetName = ObjectBaseName;
+		FString CurrentPath = PackageFolderPath;
+		if (CurrentPath.IsEmpty() == false)
+		{
+			FAssetToolsModule& AssetToolsModule = FModuleManager::LoadModuleChecked<FAssetToolsModule>("AssetTools");
+			FString UnusedPackageName;
+			AssetToolsModule.Get().CreateUniqueAssetName(PackageFolderPath + TEXT("/") + ObjectBaseName, TEXT(""), UnusedPackageName, UseDefaultAssetName);
+		}
+
+		FSaveAssetDialogConfig Config;
+		Config.DefaultAssetName = UseDefaultAssetName;
+		Config.DialogTitleOverride = LOCTEXT("GenerateTexture2DAssetPathDialogWarning", "Choose Folder Path and Name for New Asset. Cancel to Discard New Asset.");
+		Config.DefaultPath = CurrentPath;
+		FString SelectedPath = ContentBrowser.CreateModalSaveAssetDialog(Config);
+
+		if (SelectedPath.IsEmpty() == false)
+		{
+			PackageFolderPath = FPaths::GetPath(SelectedPath);
+			ObjectBaseName = FPaths::GetBaseFilename(SelectedPath, true);
+		}
+		else
+		{
+			return false;
+		}
+	}
+
+	FString UseBaseName = ObjectBaseName;
+
+	// create new package
+	FString UniqueAssetName;
+	FString UniquePackageName;
+
+	FAssetToolsModule& AssetToolsModule = FModuleManager::LoadModuleChecked<FAssetToolsModule>("AssetTools");
+	AssetToolsModule.Get().CreateUniqueAssetName(PackageFolderPath + TEXT("/") + ObjectBaseName, TEXT(""), UniquePackageName, UniqueAssetName);
+
+	UPackage* Package = CreatePackage(*UniquePackageName);
+	UPackage* AssetPackage = CreatePackage(*UniquePackageName);
+
+	// move texture from Transient package to real package
+	GeneratedTexture->Rename(*UniqueAssetName, AssetPackage, REN_None);
+	// remove transient flag, add public/standalone/transactional
+	GeneratedTexture->ClearFlags(RF_Transient);
+	GeneratedTexture->SetFlags(RF_Public | RF_Standalone | RF_Transactional);
+	// mark things as modified / dirtied
+	GeneratedTexture->Modify();
+	GeneratedTexture->UpdateResource();
+	GeneratedTexture->PostEditChange();
+	GeneratedTexture->MarkPackageDirty();
+
+	FAssetRegistryModule::AssetCreated(GeneratedTexture);
+
+	return true;
+}
+
+
+int32 UFractureToolAutoUV::ExecuteFracture(const FFractureToolContext& FractureContext)
+{
+	if (FractureContext.IsValid())
+	{
+		FScopedSlowTask UVTask(3, LOCTEXT("StartingAutoUV", "Automatically laying out and texturing internal surfaces"));
+		UVTask.MakeDialog();
+
+		FGeometryCollection& Collection = *FractureContext.GetGeometryCollection();
+		
+		int32 OutputRes = (int32)AutoUVSettings->Resolution;
+
+		UVTask.EnterProgressFrame(1, LOCTEXT("LayOutUVIslands", "Laying out UV islands"));
+		if (!UE::PlanarCut::UVLayout(Collection, OutputRes, AutoUVSettings->GutterSize))
+		{
+			// failed to do layout
+			return INDEX_NONE;
+		}
+
+		UVTask.EnterProgressFrame(1, LOCTEXT("TexturingSurfaces", "Texturing internal surfaces"));
+
+		FTexture2DBuilder TextureBuilder;
+		FImageDimensions Dimensions(OutputRes, OutputRes);
+		TextureBuilder.Initialize(FTexture2DBuilder::ETextureType::Color, Dimensions);
+
+		TImageBuilder<FVector3f> ImageBuilder;
+		ImageBuilder.SetDimensions(Dimensions);
+		ImageBuilder.Clear(FVector3f(0, 0, 0));
+
+		UE::PlanarCut::TextureInternalSurfaces(Collection, AutoUVSettings->MaxDistance, FMath::CeilToInt(AutoUVSettings->GutterSize), ImageBuilder);
+
+		UVTask.EnterProgressFrame(1, LOCTEXT("SavingTexture", "Saving result"));
+
+		TextureBuilder.Copy(ImageBuilder);
+		TextureBuilder.Commit(false);
+		AutoUVSettings->Result = TextureBuilder.GetTexture2D();
+
+		// choose default texture name based on corresponding geometry collection name
+		FString BaseName = FractureContext.GetFracturedGeometryCollection()->GetName();
+		FTexture2DBuilder::CopyPlatformDataToSourceData(AutoUVSettings->Result, FTexture2DBuilder::ETextureType::Color);
+		SaveGeneratedTexture(AutoUVSettings->Result, FString::Printf(TEXT("%s_AutoUV"), *BaseName), FractureContext.GetFracturedGeometryCollection(), AutoUVSettings->bPromptToSave);
+	}
+
+	return INDEX_NONE;
+}
+
+#undef LOCTEXT_NAMESPACE
diff --git a/Engine/Plugins/Experimental/ChaosEditor/Source/FractureEditor/Private/FractureToolUV.h b/Engine/Plugins/Experimental/ChaosEditor/Source/FractureEditor/Private/FractureToolUV.h
new file mode 100644
--- /dev/null
+++ b/Engine/Plugins/Experimental/ChaosEditor/Source/FractureEditor/Private/FractureToolUV.h
@@ -0,0 +1,99 @@
+// Copyright Epic Games, Inc. All Rights Reserved.
+
+#pragma once
+
+#include "FractureTool.h"
+
+#include "FractureToolUV.generated.h"
+
+class FFractureToolContext;
+
+UENUM()
+enum class EAutoUVTextureResolution : int32
+{
+	Resolution16 = 16 UMETA(DisplayName = "16 x 16"),
+	Resolution32 = 32 UMETA(DisplayName = "32 x 32"),
+	Resolution64 = 64 UMETA(DisplayName = "64 x 64"),
+	Resolution128 = 128 UMETA(DisplayName = "128 x 128"),
+	Resolution256 = 256 UMETA(DisplayName = "256 x 256"),
+	Resolution512 = 512 UMETA(DisplayName = "512 x 512"),
+	Resolution1024 = 1024 UMETA(DisplayName = "1024 x 1024"),
+	Resolution2048 = 2048 UMETA(DisplayName = "2048 x 2048"),
+	Resolution4096 = 4096 UMETA(DisplayName = "4096 x 4096"),
+	Resolution8192 = 8192 UMETA(DisplayName = "8192 x 8192")
+};
+
+
+/** Settings specifically related to the one-time destructive fracturing of a mesh **/
+UCLASS(config = EditorPerProjectUserSettings)
+class UFractureAutoUVSettings : public UFractureToolSettings
+{
+public:
+
+	GENERATED_BODY()
+
+	UFractureAutoUVSettings(const FObjectInitializer& ObjInit)
+		: Super(ObjInit)
+	{}
+
+	/** The pixel resolution of the generated map */
+	UPROPERTY(EditAnywhere, Category = MapSettings)
+	EAutoUVTextureResolution Resolution = EAutoUVTextureResolution::Resolution512;
+
+	/** Space to leave between UV islands, measured in texels */
+	UPROPERTY(EditAnywhere, Category = MapSettings, meta = (ClampMin = "1", ClampMax = "10", UIMax = "4"))
+	int32 GutterSize = 2;
+
+	/** Max distance to search for the outer mesh surface */
+	UPROPERTY(EditAnywhere, Category = MapSettings, meta = (UIMin = "1", UIMax = "100", ClampMin = ".01", ClampMax = "1000"))
+	double MaxDistance = 50;
+
+	/** The resulting automatically-generated texture map */
+	UPROPERTY(VisibleAnywhere, Category = MapSettings)
+	UTexture2D* Result;
+
+	/** Whether to prompt user for an asset name for each generated texture, or automatically place them next to the source geometry collections */
+	UPROPERTY(EditAnywhere, Category = MapSettings)
+	bool bPromptToSave = true;
+};
+
+
+UCLASS(DisplayName = "AutoUV Tool", Category = "FractureTools")
+class UFractureToolAutoUV : public UFractureModalTool
+{
+public:
+	GENERATED_BODY()
+
+	UFractureToolAutoUV(const FObjectInitializer& ObjInit);
+
+	///
+	/// UFractureModalTool Interface
+	///
+
+	/** This is the Text that will appear on the tool button to execute the tool **/
+	virtual FText GetDisplayText() const override;
+	virtual FText GetTooltipText() const override;
+	/** This is the Text that will appear on the button to execute the fracture **/
+	virtual FText GetApplyText() const override { return FText(NSLOCTEXT("AutoUV", "ExecuteAutoUV", "AutoUV")); }
+	virtual FSlateIcon GetToolIcon() const override;
+	virtual void RegisterUICommand(FFractureEditorCommands* BindingContext) override;
+	virtual TArray<UObject*> GetSettingsObjects() const override;
+	virtual void FractureContextChanged() override;
+	virtual void Render(const FSceneView* View, FViewport* Viewport, FPrimitiveDrawInterface* PDI) override;
+	/** Executes function that generates new geometry. Returns the first new geometry index. */
+	virtual int32 ExecuteFracture(const FFractureToolContext& FractureContext) override;
+	virtual bool CanExecute() const override;
+	virtual void PostEditChangeChainProperty(struct FPropertyChangedChainEvent& PropertyChangedEvent) override;
+
+	/** Gets the UI command info for this command */
+	const TSharedPtr<FUICommandInfo>& GetUICommandInfo() const;
+
+	virtual TArray<FFractureToolContext> GetFractureToolContexts() const override;
+
+
+protected:
+	UPROPERTY(EditAnywhere, Category = Slicing)
+	TObjectPtr<UFractureAutoUVSettings> AutoUVSettings;
+
+	bool SaveGeneratedTexture(UTexture2D* GeneratedTexture, FString ObjectBaseName, const UObject* RelativeToAsset, bool bPromptToSave);
+};
diff --git a/Engine/Plugins/Experimental/PlanarCutPlugin/Source/PlanarCut/Private/FractureAutoUV.cpp b/Engine/Plugins/Experimental/PlanarCutPlugin/Source/PlanarCut/Private/FractureAutoUV.cpp
new file mode 100644
--- /dev/null
+++ b/Engine/Plugins/Experimental/PlanarCutPlugin/Source/PlanarCut/Private/FractureAutoUV.cpp
@@ -0,0 +1,315 @@
+// Copyright Epic Games, Inc. All Rights Reserved.
+#include "FractureAutoUV.h"
+#include "PlanarCutPlugin.h"
+
+#include "Async/ParallelFor.h"
+
+#include "GeometryCollection/GeometryCollectionAlgo.h"
+
+#include "Parameterization/UVPacking.h"
+
+#include "Image/ImageOccupancyMap.h"
+
+#if WITH_EDITOR
+#include "Misc/ScopedSlowTask.h"
+#endif
+
+namespace UE { namespace PlanarCut {
+
+/**
+ * Shared helper method to set 'active' subset of triangles for a collection, to be considered for texturing
+ * @return number of active triangles
+ */
+int32 SetActiveTriangles(FGeometryCollection* Collection, TArray<bool>& ActiveTrianglesOut,
+	bool bActivateInsideTriangles,
+	bool bOddMaterialAreInside, TArrayView<int32> WhichMaterialsAreInside)
+{
+	TSet<int32> InsideMaterials;
+	for (int32 MatID : WhichMaterialsAreInside)
+	{
+		InsideMaterials.Add(MatID);
+	}
+
+	ActiveTrianglesOut.Init(false, Collection->Indices.Num());
+	int32 NumTriangles = 0;
+	for (int TID = 0; TID < Collection->Indices.Num(); TID++)
+	{
+		if (!Collection->Visible[TID])
+		{
+			continue;
+		}
+
+		int32 MaterialID = Collection->MaterialID[TID];
+		if (bOddMaterialAreInside && ((MaterialID % 2) == 0) == bActivateInsideTriangles)
+		{
+			continue;
+		}
+		if (WhichMaterialsAreInside.Num() > 0 && (InsideMaterials.Contains(MaterialID) != bActivateInsideTriangles))
+		{
+			continue;
+		}
+
+		ActiveTrianglesOut[TID] = true;
+		NumTriangles++;
+	}
+
+	return NumTriangles;
+}
+
+// Adapter to use geometry collections in the UVPacker, and other GeometricObjects generic mesh processing code
+// Note: Adapts the whole geometry collection as a single mesh, with triangles filtered by an ActiveTriangles mask
+struct FGeomMesh : public FUVPacker::IUVMeshView
+{
+	FGeometryCollection* Collection;
+	const TArrayView<bool> ActiveTriangles;
+	int32 NumTriangles;
+	TArray<FVector3d> GlobalVertices; // vertices transformed to global space
+
+	// Construct a mesh from a geometry collection and a mask of active triangles
+	FGeomMesh(FGeometryCollection* Collection, TArrayView<bool> ActiveTriangles, int32 NumTriangles) 
+		: Collection(Collection), ActiveTriangles(ActiveTriangles), NumTriangles(NumTriangles)
+	{
+		InitVertices();
+	}
+
+	// Construct a mesh from an existing FGeomMesh and a mask of active triangles
+	FGeomMesh(const FGeomMesh& OtherMesh, TArrayView<bool> ActiveTriangles, int32 NumTriangles)
+		: Collection(OtherMesh.Collection), ActiveTriangles(ActiveTriangles), NumTriangles(NumTriangles)
+	{
+		GlobalVertices = OtherMesh.GlobalVertices;
+	}
+
+	void InitVertices()
+	{
+		// transform all vertices from local to global space
+		TArray<FTransform> GlobalTransformArray;
+		GeometryCollectionAlgo::GlobalMatrices(Collection->Transform, Collection->Parent, GlobalTransformArray);
+		GlobalVertices.SetNum(Collection->Vertex.Num());
+		for (int32 Idx = 0; Idx < GlobalVertices.Num(); Idx++)
+		{
+			GlobalVertices[Idx] = GlobalTransformArray[Collection->BoneMap[Idx]].TransformPosition(Collection->Vertex[Idx]);
+		}
+	}
+
+	///
+	/// The FUVPacker::IUVMeshView interface functions; this is the subset of functionality required for FUVPacker
+	/// 
+
+	virtual FIndex3i GetTriangle(int32 TID) const
+	{
+		return FIndex3i(Collection->Indices[TID]);
+	}
+
+	virtual FIndex3i GetUVTriangle(int32 TID) const
+	{
+		return FIndex3i(Collection->Indices[TID]);
+	}
+
+	virtual FVector3d GetVertex(int32 VID) const
+	{
+		return GlobalVertices[VID];
+	}
+
+	virtual FVector2f GetUV(int32 EID) const
+	{
+		return FVector2f(Collection->UV[EID]);
+	}
+
+	virtual void SetUV(int32 EID, FVector2f UVIn)
+	{
+		FVector2D& UV = Collection->UV[EID];
+		UV.X = UVIn.X;
+		UV.Y = UVIn.Y;
+	}
+
+
+
+	///
+	/// Interface for templated mesh code, not part of the the IUVMeshView interface
+	///
+	inline bool IsTriangle(int32 TID) const
+	{
+		return ActiveTriangles[TID];
+	}
+	constexpr inline bool IsVertex(int32 VID) const
+	{
+		return true; // we don't have sparse vertices, so can always return true
+	}
+	inline int32 MaxTriangleID() const
+	{
+		return Collection->Indices.Num();
+	}
+	inline int32 MaxVertexID() const
+	{
+		return GlobalVertices.Num();
+	}
+	inline int32 TriangleCount() const
+	{
+		return NumTriangles;
+	}
+	inline int32 VertexCount() const
+	{
+		return GlobalVertices.Num();
+	}
+	constexpr inline int32 GetShapeTimestamp() const
+	{
+		return 0;
+	}
+
+	inline void GetTriVertices(int TID, FVector3d& V0, FVector3d& V1, FVector3d& V2) const
+	{
+		FIntVector TriRaw = Collection->Indices[TID];
+
+		V0 = GlobalVertices[TriRaw.X];
+		V1 = GlobalVertices[TriRaw.Y];
+		V2 = GlobalVertices[TriRaw.Z];
+	}
+};
+
+/// Like FGeomMesh, but the UVs are given as the vertices rather than accessed separately
+/// Used by the texture sampling code
+struct FGeomFlatUVMesh
+{
+	FGeometryCollection* Collection;
+	const TArrayView<bool> ActiveTriangles;
+	int32 NumTriangles;
+
+	FGeomFlatUVMesh(FGeometryCollection* Collection, TArrayView<bool> ActiveTriangles, int32 NumTriangles)
+		: Collection(Collection), ActiveTriangles(ActiveTriangles), NumTriangles(NumTriangles)
+	{}
+
+	inline FIndex3i GetTriangle(int32 TID) const
+	{
+		return FIndex3i(Collection->Indices[TID]);
+	}
+
+	inline FVector3d GetVertex(int32 VID) const
+	{
+		const FVector2D& UV = Collection->UV[VID];
+		return FVector3d(UV.X, UV.Y, 0);
+	}
+
+	inline bool IsTriangle(int32 TID) const
+	{
+		return ActiveTriangles[TID];
+	}
+	constexpr inline bool IsVertex(int32 VID) const
+	{
+		return true; // we don't have sparse vertices, so can always return true
+	}
+	inline int32 MaxTriangleID() const
+	{
+		return Collection->Indices.Num();
+	}
+	inline int32 MaxVertexID() const
+	{
+		return Collection->UV.Num();
+	}
+	inline int32 TriangleCount() const
+	{
+		return NumTriangles;
+	}
+	inline int32 VertexCount() const
+	{
+		return Collection->UV.Num();
+	}
+	constexpr inline int32 GetShapeTimestamp() const
+	{
+		return 0;
+	}
+
+	inline void GetTriVertices(int TID, FVector3d& V0, FVector3d& V1, FVector3d& V2) const
+	{
+		FIntVector TriRaw = Collection->Indices[TID];
+
+		V0 = GetVertex(TriRaw.X);
+		V1 = GetVertex(TriRaw.Y);
+		V2 = GetVertex(TriRaw.Z);
+	}
+};
+
+
+bool UVLayout(
+	FGeometryCollection& Collection,
+	int32 UVRes,
+	float GutterSize,
+	bool bOnlyOddMaterials,
+	TArrayView<int32> WhichMaterials
+)
+{
+	TArray<bool> ActiveTriangles;
+	int32 NumActive = SetActiveTriangles(&Collection, ActiveTriangles, true, bOnlyOddMaterials, WhichMaterials);
+	FGeomMesh UVMesh(&Collection, ActiveTriangles, NumActive);
+
+	TArray<TArray<int32>> UVIslands;
+	UE::UVPacking::CreateUVIslandsFromMeshTopology<FGeomMesh>(UVMesh, UVIslands);
+
+	FUVPacker Packer;
+	Packer.bAllowFlips = false;
+	// StandardPack doesn't support the Packer.GutterSize feature, and always tries to leave a 1 texel gutter.
+	// To approximate a larger gutter, we tell it to consider a smaller output resolution --
+	//  hoping that e.g. the UVs for a 1 pixel gutter at 256x256 are ~ a 2 pixel gutter at 512x512
+	// TODO: If we make StandardPack support the GutterSize parameter, we can use that instead.
+	Packer.TextureResolution = UVRes / FMathf::Max(1, GutterSize);
+	return Packer.StandardPack(&UVMesh, UVIslands);
+}
+
+void TextureInternalSurfaces(
+	FGeometryCollection& Collection,
+	double MaxDistance,
+	int32 GutterSize,
+	TImageBuilder<FVector3f>& TextureOut,
+	bool bOnlyOddMaterials,
+	TArrayView<int32> WhichMaterials
+)
+{
+	TArray<bool> InsideTriangles;
+	int32 NumInsideTris = SetActiveTriangles(&Collection, InsideTriangles, true, bOnlyOddMaterials, WhichMaterials);
+	FGeomFlatUVMesh UVMesh(&Collection, InsideTriangles, NumInsideTris);
+
+	FImageOccupancyMap OccupancyMap;
+	OccupancyMap.GutterSize = GutterSize;
+	OccupancyMap.Initialize(TextureOut.GetDimensions());
+	OccupancyMap.ComputeFromUVSpaceMesh(UVMesh, [](int32 TriangleID) { return TriangleID; });
+
+	FGeomMesh InsideMesh(&Collection, InsideTriangles, NumInsideTris);
+
+	TArray<bool> OutsideTriangles;
+	int32 NumOutsideTris = SetActiveTriangles(&Collection, OutsideTriangles, false, bOnlyOddMaterials, WhichMaterials);
+	FGeomMesh OutsideMesh(InsideMesh, OutsideTriangles, NumOutsideTris);
+	TMeshAABBTree3<FGeomMesh> OutsideSpatial(&OutsideMesh);
+
+	ParallelFor(OccupancyMap.Dimensions.GetHeight(),
+		[&MaxDistance, &TextureOut, &UVMesh, &OccupancyMap, &InsideMesh, &OutsideSpatial](int32 Y)
+	{
+		for (int32 X = 0; X < OccupancyMap.Dimensions.GetWidth(); X++)
+		{
+			int64 LinearCoord = OccupancyMap.Dimensions.GetIndex(X, Y);
+			if (OccupancyMap.IsInterior(LinearCoord))
+			{
+				int32 TID = OccupancyMap.TexelQueryTriangle[LinearCoord];
+				FVector2f UV = OccupancyMap.TexelQueryUV[LinearCoord];
+				// Note this is the slowest way to get barycentric coordinates out of a point and a 2D triangle,
+				//   ... but it's the one we currently have that's robust to (and gives good answers on) degenerate triangles
+				FDistPoint3Triangle3d DistQuery = TMeshQueries<FGeomFlatUVMesh>::TriangleDistance(UVMesh, TID, FVector3d(UV.X, UV.Y, 0));
+				FVector3d Bary = DistQuery.TriangleBaryCoords;
+				FTriangle3d Tri;
+				InsideMesh.GetTriVertices(TID, Tri.V[0], Tri.V[1], Tri.V[2]);
+				FVector3d InsidePoint = Tri.BarycentricPoint(Bary);
+				double DistanceSq;
+				OutsideSpatial.FindNearestTriangle(InsidePoint, DistanceSq, MaxDistance);
+				float PercentDistance = FMathf::Min(1.0f, FMathf::Sqrt(DistanceSq) / (float)MaxDistance);
+				ensure(FMath::IsFinite(PercentDistance));
+
+				TextureOut.SetPixel(LinearCoord, FVector3f(PercentDistance, 0, 1));
+			}
+		}
+	});
+
+	for (TTuple<int64, int64>& GutterToInside : OccupancyMap.GutterTexels)
+	{
+		TextureOut.CopyPixel(GutterToInside.Value, GutterToInside.Key);
+	}
+}
+
+}} // namespace UE::PlanarCut
\ No newline at end of file
diff --git a/Engine/Plugins/Experimental/PlanarCutPlugin/Source/PlanarCut/Public/FractureAutoUV.h b/Engine/Plugins/Experimental/PlanarCutPlugin/Source/PlanarCut/Public/FractureAutoUV.h
new file mode 100644
--- /dev/null
+++ b/Engine/Plugins/Experimental/PlanarCutPlugin/Source/PlanarCut/Public/FractureAutoUV.h
@@ -0,0 +1,49 @@
+// Copyright Epic Games, Inc. All Rights Reserved.
+#pragma once
+#include "CoreMinimal.h"
+
+#include "GeometryCollection/GeometryCollection.h"
+
+#include "Image/ImageBuilder.h"
+
+namespace UE { namespace PlanarCut {
+
+/**
+ * Make a UV atlas of non-overlapping UV charts for a geometry collection
+ *
+ * @param Collection		The collection to be atlas'd
+ * @param UVRes				Target resolution for the atlas
+ * @param GutterSize		Space to leave between UV islands, in pixels at the target resolution
+ * @param bOnlyOddMaterials	If true, restrict UV island layout to odd-numbered material IDs
+ * @param WhichMaterials	If non-empty, restrict UV island layout to only the listed material IDs
+ */
+bool PLANARCUT_API UVLayout(
+	FGeometryCollection& Collection,
+	int32 UVRes = 1024,
+	float GutterSize = 1,
+	bool bOnlyOddMaterials = true,
+	TArrayView<int32> WhichMaterials = TArrayView<int32>()
+);
+
+
+/**
+ * Generate a texture for internal faces based on depth inside surface
+ * TODO: add options to texture based on other quantities
+ *
+ * @param Collection		The collection to be create a new texture for
+ * @param MaxDistance		Maximum distance to search for 'outside' surface, when computing depth inside surface
+ * @param GutterSize		Number of texels to fill outside of UV island borders (values are copied from nearest inside pt)
+ * @param TextureOut		Texture to write to
+ * @param bOnlyOddMaterials	If true, restrict UV island layout to odd-numbered material IDs
+ * @param WhichMaterials	If non-empty, restrict UV island layout to only the listed material IDs
+ */
+void PLANARCUT_API TextureInternalSurfaces(
+	FGeometryCollection& Collection,
+	double MaxDistance,
+	int32 GutterSize,
+	TImageBuilder<FVector3f>& TextureOut,
+	bool bOnlyOddMaterials = true,
+	TArrayView<int32> WhichMaterials = TArrayView<int32>()
+);
+
+}} // namespace UE::PlanarCut
\ No newline at end of file
