fix fracture crash on appending non-manifold triangle when creating a dynamic mesh version of a geometry collection

#rnx

diff --git a/Engine/Plugins/Experimental/PlanarCutPlugin/Source/PlanarCut/Private/PlanarCut.cpp b/Engine/Plugins/Experimental/PlanarCutPlugin/Source/PlanarCut/Private/PlanarCut.cpp
--- a/Engine/Plugins/Experimental/PlanarCutPlugin/Source/PlanarCut/Private/PlanarCut.cpp
+++ b/Engine/Plugins/Experimental/PlanarCutPlugin/Source/PlanarCut/Private/PlanarCut.cpp
@@ -661,7 +661,10 @@ private:
 							Swap(Tri.B, Tri.C);
 						}
 						int TID = Meshes[MeshIdx]->AppendTriangle(Tri);
-						Meshes[MeshIdx]->Attributes()->GetMaterialID()->SetNewValue(TID, MID);
+						if (ensure(TID > -1))
+						{
+							Meshes[MeshIdx]->Attributes()->GetMaterialID()->SetNewValue(TID, MID);
+						}
 					}
 				}
 			}
@@ -693,7 +696,10 @@ private:
 							Swap(Triangle.B, Triangle.C);
 						}
 						int TID = Meshes[MeshIdx]->AppendTriangle(Triangle);
-						Meshes[MeshIdx]->Attributes()->GetMaterialID()->SetNewValue(TID, MID);
+						if (ensure(TID > -1))
+						{
+							Meshes[MeshIdx]->Attributes()->GetMaterialID()->SetNewValue(TID, MID);
+						}
 					}
 				}
 			}
@@ -802,7 +808,10 @@ private:
 				for (FIndex3i Triangle : Triangulation.Triangles)
 				{
 					int TID = Mesh.AppendTriangle(Triangle);
-					Mesh.Attributes()->GetMaterialID()->SetNewValue(TID, MID);
+					if (ensure(TID > -1))
+					{
+						Mesh.Attributes()->GetMaterialID()->SetNewValue(TID, MID);
+					}
 				}
 
 				RemeshForNoise(Mesh, EEdgeRefineFlags::SplitsOnly, Spacing);
@@ -1290,7 +1299,28 @@ struct FDynamicMeshCollection
 			FIntVector VertexOffset(VertexStart, VertexStart, VertexStart);
 			for (int32 Idx = Collection->FaceStart[GeometryIdx], N = Collection->FaceStart[GeometryIdx] + FaceCount; Idx < N; Idx++)
 			{
-				int TID = Mesh.AppendTriangle(FIndex3i(Collection->Indices[Idx] - VertexOffset), 0);
+				FIndex3i AddTri = FIndex3i(Collection->Indices[Idx] - VertexOffset);
+				int TID = Mesh.AppendTriangle(AddTri, 0);
+				if (TID == FDynamicMesh3::NonManifoldID)
+				{
+					// work around non-manifold triangles by copying the vertices
+					FIndex3i NewTri(-1, -1, -1);
+					for (int SubIdx = 0; SubIdx < 3; SubIdx++)
+					{
+						int NewVID = Mesh.AppendVertex(Mesh, AddTri[SubIdx]);
+						int32 SrcIdx = AddTri[SubIdx] + VertexStart;
+						UE::PlanarCutInternals::AugmentDynamicMesh::SetTangent(Mesh, NewVID,
+							Mesh.GetVertexNormal(NewVID), // TODO: we don't actually use the vertex normal; consider removing this arg from the function entirely
+							CollectionToLocal.TransformVectorNoScale(FVector3f(Collection->TangentU[SrcIdx])),
+							CollectionToLocal.TransformVectorNoScale(FVector3f(Collection->TangentV[SrcIdx])));
+						NewTri[SubIdx] = NewVID;
+					}
+					TID = Mesh.AppendTriangle(NewTri, 0);
+				}
+				if (TID < 0)
+				{
+					continue;
+				}
 				Mesh.Attributes()->GetMaterialID()->SetValue(TID, Collection->MaterialID[Idx]);
 				UE::PlanarCutInternals::AugmentDynamicMesh::SetVisibility(Mesh, TID, Collection->Visible[Idx]);
 				// note: material index doesn't need to be passed through; will be rebuilt by a call to reindex materials once the cut mesh is returned back to geometry collection format
@@ -2690,9 +2720,12 @@ int32 CutWithMesh(
 		for (int TID = 0; TID < Triangles.Num(); TID++)
 		{
 			FIndex3i& Tri = Triangles[TID];
-			OutMesh.AppendTriangle(Tri);
-			OutMaterialID->SetValue(TID, -1); // just use a single negative material ID by convention to indicate internal material
-			UE::PlanarCutInternals::AugmentDynamicMesh::SetVisibility(OutMesh, TID, true);
+			int AddedTID = OutMesh.AppendTriangle(Tri);
+			if (ensure(AddedTID > -1))
+			{
+				OutMaterialID->SetValue(AddedTID, -1); // just use a single negative material ID by convention to indicate internal material
+				UE::PlanarCutInternals::AugmentDynamicMesh::SetVisibility(OutMesh, AddedTID, true);
+			}
 		}
 	}
 
