fracture crash fixes: fix resample to directly update geometry rather than delete transforms + fix memory clobbering due to tarray of not-trivially-movable things

#rnx

diff --git a/Engine/Plugins/Experimental/PlanarCutPlugin/Source/PlanarCut/Private/PlanarCut.cpp b/Engine/Plugins/Experimental/PlanarCutPlugin/Source/PlanarCut/Private/PlanarCut.cpp
--- a/Engine/Plugins/Experimental/PlanarCutPlugin/Source/PlanarCut/Private/PlanarCut.cpp
+++ b/Engine/Plugins/Experimental/PlanarCutPlugin/Source/PlanarCut/Private/PlanarCut.cpp
@@ -281,7 +281,7 @@ struct FCellMeshes
 		// TODO: compute spatial in advance?  (only useful if we rework mesh booleans to support it)
 		//FDynamicMeshAABBTree3 Spatial;
 	};
-	TArray<FCellInfo> CellMeshes;
+	TIndirectArray<FCellInfo> CellMeshes;
 	int32 OutsideCellIndex = -1;
 
 	// Noise Offsets, to randomize where perlin noise is sampled
@@ -289,6 +289,15 @@ struct FCellMeshes
 	FVector NoiseOffsetY;
 	FVector NoiseOffsetZ;
 
+	void SetNumCells(int32 NumMeshes)
+	{
+		CellMeshes.Reset();
+		for (int32 Idx = 0; Idx < NumMeshes; Idx++)
+		{
+			CellMeshes.Add(new FCellInfo);
+		}
+	}
+
 	FCellMeshes()
 	{
 		InitEmpty();
@@ -301,8 +310,7 @@ struct FCellMeshes
 
 	FCellMeshes(FDynamicMesh3& SingleCutter, const FInternalSurfaceMaterials& Materials, TOptional<FTransform> Transform)
 	{
-		CellMeshes.Reset();
-		CellMeshes.SetNum(2);
+		SetNumCells(2);
 
 		if (Transform.IsSet())
 		{
@@ -341,7 +349,7 @@ struct FCellMeshes
 			GlobalUVScale = 1;
 		}
 
-		CellMeshes.SetNum(1);
+		SetNumCells(1);
 
 		bool bNoise = Cells.InternalSurfaceMaterials.NoiseSettings.IsSet();
 
@@ -478,8 +486,7 @@ struct FCellMeshes
 			NumCells++;
 		}
 
-		CellMeshes.Reset();
-		CellMeshes.SetNum(NumCells);
+		SetNumCells(NumCells);
 
 		bool bNoise = Cells.InternalSurfaceMaterials.NoiseSettings.IsSet();
 		if (bNoise)
@@ -1245,7 +1252,7 @@ struct FDynamicMeshCollection
 		FMeshData(const FDynamicMesh3& Mesh, int32 TransformIndex, FTransform ToCollection) : AugMesh(Mesh), TransformIndex(TransformIndex), ToCollection(ToCollection)
 		{}
 	};
-	TArray<FMeshData> Meshes;
+	TIndirectArray<FMeshData> Meshes;
 	FAxisAlignedBox3d Bounds;
 
 	FDynamicMeshCollection(const FGeometryCollection* Collection, const TArrayView<const int32>& TransformIndices, FTransform TransformCollection)
@@ -1268,7 +1275,8 @@ struct FDynamicMeshCollection
 
 			FTransform3d CollectionToLocal = FTransform3d(GeometryCollectionAlgo::GlobalMatrix(Collection->Transform, Collection->Parent, TransformIdx) * TransformCollection);
 
-			FMeshData& MeshData = Meshes.Emplace_GetRef();
+			int32 AddedMeshIdx = Meshes.Add(new FMeshData);
+			FMeshData& MeshData = Meshes[AddedMeshIdx];
 			MeshData.TransformIndex = TransformIdx;
 			MeshData.ToCollection = FTransform(CollectionToLocal.Inverse());
 			FDynamicMesh3& Mesh = MeshData.AugMesh;
@@ -1373,7 +1381,7 @@ struct FDynamicMeshCollection
 			//    (Note the outside cell mesh is subtracted, not intersected)
 			//    (Note this relies on island splitting to separate all the pieces afterwards.)
 			FCellMeshes GroutCells;
-			GroutCells.CellMeshes.SetNum(2);
+			GroutCells.SetNumCells(2);
 			FDynamicMesh3& GroutMesh = GroutCells.CellMeshes[0].AugMesh;
 			FDynamicMeshEditor GroutAppender(&GroutMesh);
 			FMeshIndexMappings IndexMaps;
@@ -1979,6 +1987,77 @@ struct FDynamicMeshCollection
 		return Output.BoneName[TransformParent] + "_" + FString::FromInt(SubPartIndex);
 	}
 
+	void AddCollisionSamples(double CollisionSampleSpacing)
+	{
+		for (int32 MeshIdx = 0; MeshIdx < Meshes.Num(); MeshIdx++)
+		{
+			UE::PlanarCutInternals::AugmentDynamicMesh::AddCollisionSamplesPerComponent(Meshes[MeshIdx].AugMesh, CollisionSampleSpacing);
+		}
+	}
+
+	// Update an existing geometry in a collection w/ a new mesh (w/ the same number of faces and vertices!)
+	static bool UpdateCollection(const FTransform& ToCollection, FDynamicMesh3& Mesh, int32 GeometryIdx, FGeometryCollection& Output, int32 InternalMaterialID)
+	{
+		if (!Mesh.IsCompact())
+		{
+			Mesh.CompactInPlace(nullptr);
+		}
+
+		int32 OldVertexCount = Output.VertexCount[GeometryIdx];
+		int32 OldTriangleCount = Output.FaceCount[GeometryIdx];
+
+		int32 NewVertexCount = Mesh.VertexCount();
+		int32 NewTriangleCount = Mesh.TriangleCount();
+
+		if (!ensure(OldVertexCount == NewVertexCount) || !ensure(OldTriangleCount == NewTriangleCount))
+		{
+			return false;
+		}
+
+		int32 VerticesStart = Output.VertexStart[GeometryIdx];
+		int32 FacesStart = Output.FaceStart[GeometryIdx];
+		int32 TransformIdx = Output.TransformIndex[GeometryIdx];
+
+		for (int32 VID = 0; VID < Mesh.MaxVertexID(); VID++)
+		{
+			checkSlow(Mesh.IsVertex(VID)); // mesh is compact
+			int32 CopyToIdx = VerticesStart + VID;
+			Output.Vertex[CopyToIdx] = ToCollection.TransformPosition(FVector(Mesh.GetVertex(VID)));
+			Output.Normal[CopyToIdx] = ToCollection.TransformVectorNoScale(FVector(Mesh.GetVertexNormal(VID)));
+			Output.UV[CopyToIdx] = FVector2D(Mesh.GetVertexUV(VID));
+			FVector3f TangentU, TangentV;
+			UE::PlanarCutInternals::AugmentDynamicMesh::GetTangent(Mesh, VID, TangentU, TangentV);
+			Output.TangentU[CopyToIdx] = ToCollection.TransformVectorNoScale(FVector(TangentU));
+			Output.TangentV[CopyToIdx] = ToCollection.TransformVectorNoScale(FVector(TangentV));
+			Output.Color[CopyToIdx] = FVector(Mesh.GetVertexColor(VID));
+
+			// Bone map is set based on the transform of the new geometry
+			Output.BoneMap[CopyToIdx] = TransformIdx;
+		}
+
+		FIntVector VertexStartOffset(VerticesStart);
+		for (int32 TID = 0; TID < Mesh.MaxTriangleID(); TID++)
+		{
+			checkSlow(Mesh.IsTriangle(TID));
+			int32 CopyToIdx = FacesStart + TID;
+			Output.Visible[CopyToIdx] = UE::PlanarCutInternals::AugmentDynamicMesh::GetVisibility(Mesh, TID);
+			int MaterialID = Mesh.Attributes()->GetMaterialID()->GetValue(TID);
+			Output.MaterialID[CopyToIdx] = MaterialID < 0 ? InternalMaterialID : MaterialID;
+			Output.Indices[CopyToIdx] = FIntVector(Mesh.GetTriangle(TID)) + VertexStartOffset;
+		}
+
+		if (Output.BoundingBox.Num())
+		{
+			Output.BoundingBox[GeometryIdx].Init();
+			for (int32 Idx = VerticesStart; Idx < VerticesStart + Output.VertexCount[GeometryIdx]; ++Idx)
+			{
+				Output.BoundingBox[GeometryIdx] += Output.Vertex[Idx];
+			}
+		}
+
+		return true;
+	}
+
 	static int32 AppendToCollection(const FTransform& ToCollection, FDynamicMesh3& Mesh, double CollisionSampleSpacing, int32 TransformParent, FString BoneName, FGeometryCollection& Output, int32 InternalMaterialID)
 	{
 		if (Mesh.TriangleCount() == 0)
@@ -2047,7 +2126,7 @@ struct FDynamicMeshCollection
 		FIntVector VertexStartOffset(VerticesStart);
 		for (int32 TID = 0; TID < Mesh.MaxTriangleID(); TID++)
 		{
-			check(Mesh.IsTriangle(TID));
+			checkSlow(Mesh.IsTriangle(TID));
 			int32 CopyToIdx = FacesStart + TID;
 			Output.Visible[CopyToIdx] = UE::PlanarCutInternals::AugmentDynamicMesh::GetVisibility(Mesh, TID);
 			int MaterialID = Mesh.Attributes()->GetMaterialID()->GetValue(TID);
@@ -2763,35 +2842,38 @@ int32 AddCollisionSampleVertices(double CollisionSampleSpacing, FGeometryCollect
 	FTransform CellsToWorld = FTransform::Identity;
 
 	FDynamicMeshCollection MeshCollection(&Collection, TransformIndices, CellsToWorld);
-	struct FLocalInfo
+
+	MeshCollection.AddCollisionSamples(CollisionSampleSpacing);
+
+	int32 NumGeometry = Collection.NumElements(FGeometryCollection::GeometryGroup);
+	TArray<int32> NewFaceCounts, NewVertexCounts;
+	NewFaceCounts.SetNumUninitialized(NumGeometry);
+	NewVertexCounts.SetNumUninitialized(NumGeometry);
+	for (int32 GeomIdx = 0; GeomIdx < Collection.FaceCount.Num(); GeomIdx++)
 	{
-		FString BoneName;
-		int Parent = -1;
-	};
-	TArray<FLocalInfo> LocalInfo; LocalInfo.SetNum(TransformIndices.Num());
-	for (int Idx = 0; Idx < TransformIndices.Num(); Idx++)
+		NewFaceCounts[GeomIdx] = Collection.FaceCount[GeomIdx];
+		NewVertexCounts[GeomIdx] = Collection.VertexCount[GeomIdx];
+	}
+	for (int MeshIdx = 0; MeshIdx < MeshCollection.Meshes.Num(); MeshIdx++)
 	{
-		LocalInfo[Idx].BoneName = Collection.BoneName[TransformIndices[Idx]];
-		LocalInfo[Idx].Parent = Collection.Parent[TransformIndices[Idx]];
+		FDynamicMeshCollection::FMeshData& MeshData = MeshCollection.Meshes[MeshIdx];
+		int32 GeomIdx = Collection.TransformToGeometryIndex[MeshData.TransformIndex];
+		NewFaceCounts[GeomIdx] = MeshData.AugMesh.TriangleCount();
+		NewVertexCounts[GeomIdx] = MeshData.AugMesh.VertexCount();
 	}
+	GeometryCollectionAlgo::ResizeGeometries(&Collection, NewFaceCounts, NewVertexCounts);
 
-	TArray<int32> ToDelete(TransformIndices);
-	ToDelete.Sort();
-	Collection.RemoveElements(FGeometryCollection::TransformGroup, ToDelete);
-	int32 FirstAppended = -1;
 	for (int MeshIdx = 0; MeshIdx < MeshCollection.Meshes.Num(); MeshIdx++)
 	{
 		FDynamicMeshCollection::FMeshData& MeshData = MeshCollection.Meshes[MeshIdx];
 		FDynamicMesh3& Mesh = MeshData.AugMesh;
-		int32 ResultIdx = MeshCollection.AppendToCollection(MeshData.ToCollection, Mesh, CollisionSampleSpacing, LocalInfo[MeshIdx].Parent, LocalInfo[MeshIdx].BoneName, Collection, -1);
-		if (FirstAppended == -1)
-		{
-			FirstAppended = ResultIdx;
-		}
+		int32 GeometryIdx = Collection.TransformToGeometryIndex[MeshData.TransformIndex];
+		ensure(MeshCollection.UpdateCollection(MeshData.ToCollection, Mesh, GeometryIdx, Collection, -1));
 	}
 
 	Collection.ReindexMaterials();
 
-	return FirstAppended;
+	// TODO: This function does not create any new bones, so we could change it to not return anything
+	return INDEX_NONE;
 }
 
diff --git a/Engine/Source/Runtime/Experimental/Chaos/Private/GeometryCollection/GeometryCollectionAlgo.cpp b/Engine/Source/Runtime/Experimental/Chaos/Private/GeometryCollection/GeometryCollectionAlgo.cpp
--- a/Engine/Source/Runtime/Experimental/Chaos/Private/GeometryCollection/GeometryCollectionAlgo.cpp
+++ b/Engine/Source/Runtime/Experimental/Chaos/Private/GeometryCollection/GeometryCollectionAlgo.cpp
@@ -748,11 +748,126 @@ namespace GeometryCollectionAlgo
 			GeometryCollection->FaceStart[GeometryIdx] = NewStart;
 		}
 		GeometryCollection->ReorderElements(FGeometryCollection::FacesGroup, ShiftedOrder);
-		// fix face starts post-reorder (they will have been helpfully 'fixed' to point at their previous elements by the reorder call)
-		for (int32 GeometryIdx = 1; GeometryIdx < NumGeometries; GeometryIdx++)
+	}
+
+	void ResizeGeometries(FGeometryCollection* GeometryCollection, const TArray<int32>& FaceCounts, const TArray<int32>& VertexCounts)
+	{
+		check(GeometryCollection);
+		int32 NumGeometries = GeometryCollection->NumElements(FGeometryCollection::GeometryGroup);
+		check(FaceCounts.Num() == NumGeometries);
+		check(VertexCounts.Num() == NumGeometries);
+
+		int32 NewNumFaces = 0, NewNumVertices = 0;
+		TArray<int32> UnusedFaces, UnusedVertices;
+
+		auto AddRange = [](TArray<int32>& OutArray, int32 StartIncl, int32 EndExcl)
+		{
+			for (int32 Idx = StartIncl; Idx < EndExcl; Idx++)
+			{
+				OutArray.Add(Idx);
+			}
+		};
+
+		auto CountElements = [&AddRange](int32 OldTotal, const TArray<int32>& NewCounts, const TManagedArray<int32>& OldCounts, const TManagedArray<int32>& OldStarts, int32& OutTotal, TArray<int32>& OutUnused)
+		{
+			OutTotal = 0;
+			OutUnused.Reset();
+			for (int32 Idx = 0; Idx < NewCounts.Num(); Idx++)
+			{
+				OutTotal += NewCounts[Idx];
+				if (OldCounts[Idx] > NewCounts[Idx])
+				{
+					AddRange(OutUnused, OldStarts[Idx] + NewCounts[Idx], OldStarts[Idx] + OldCounts[Idx]);
+				}
+			}
+			if (OutTotal > OldTotal)
+			{
+				AddRange(OutUnused, OldTotal, OutTotal);
+			}
+		};
+
+		int32 OldNumFaces = GeometryCollection->NumElements(FGeometryCollection::FacesGroup);
+		int32 OldNumVertices = GeometryCollection->NumElements(FGeometryCollection::VerticesGroup);
+		CountElements(OldNumFaces, FaceCounts, GeometryCollection->FaceCount, GeometryCollection->FaceStart, NewNumFaces, UnusedFaces);
+		CountElements(OldNumVertices, VertexCounts, GeometryCollection->VertexCount, GeometryCollection->VertexStart, NewNumVertices, UnusedVertices);
+
+		// add elements at the end if needed
+		if (NewNumFaces > OldNumFaces)
+		{
+			GeometryCollection->AddElements(NewNumFaces - OldNumFaces, FGeometryCollection::FacesGroup);
+		}
+		if (NewNumVertices > OldNumVertices)
+		{
+			GeometryCollection->AddElements(NewNumVertices - OldNumVertices, FGeometryCollection::VerticesGroup);
+		}
+		
+		int32 NumFaces = GeometryCollection->NumElements(FGeometryCollection::FacesGroup);
+		int32 NumVertices = GeometryCollection->NumElements(FGeometryCollection::VerticesGroup);
+
+		auto ComputeNewOrder = [&AddRange](int32 Total, const TArray<int32>& NewCounts, const TManagedArray<int32>& OldCounts, const TManagedArray<int32>& OldStarts, TArray<int32>& Unused, TArray<int32>& OutNewOrder)
 		{
-			GeometryCollection->FaceStart[GeometryIdx] = GeometryCollection->FaceStart[GeometryIdx - 1] + GeometryCollection->FaceCount[GeometryIdx - 1];
+			OutNewOrder.Reset();
+			OutNewOrder.Reserve(Total);
+			int32 UnusedIdx = 0;
+			for (int32 GeometryIdx = 0; GeometryIdx < OldStarts.Num(); GeometryIdx++)
+			{
+				if (NewCounts[GeometryIdx] < OldCounts[GeometryIdx]) // Geometry shrunk
+				{
+					AddRange(OutNewOrder, OldStarts[GeometryIdx], OldStarts[GeometryIdx] + NewCounts[GeometryIdx]);
+				}
+				else // Geometry grew
+				{
+					AddRange(OutNewOrder, OldStarts[GeometryIdx], OldStarts[GeometryIdx] + OldCounts[GeometryIdx]);
+					int32 GrowAmt = NewCounts[GeometryIdx] - OldCounts[GeometryIdx];
+					for (int32 Idx = 0; Idx < GrowAmt; Idx++)
+					{
+						OutNewOrder.Add(Unused[UnusedIdx++]);
+					}
+				}
+			}
+			while (OutNewOrder.Num() < Total)
+			{
+				OutNewOrder.Add(Unused[UnusedIdx++]);
+			}
+		};
+
+		TArray<int32> NewFaceOrder, NewVertexOrder;
+		ComputeNewOrder(NumFaces, FaceCounts, GeometryCollection->FaceCount, GeometryCollection->FaceStart, UnusedFaces, NewFaceOrder);
+		ComputeNewOrder(NumVertices, VertexCounts, GeometryCollection->VertexCount, GeometryCollection->VertexStart, UnusedVertices, NewVertexOrder);
+		GeometryCollection->ReorderElements(FGeometryCollection::VerticesGroup, NewVertexOrder);
+		GeometryCollection->ReorderElements(FGeometryCollection::FacesGroup, NewFaceOrder);
+
+		// fix face/vertex counts and vertex->transform (bone) map
+		for (int32 GeometryIdx = 0; GeometryIdx < NumGeometries; GeometryIdx++)
+		{
+			GeometryCollection->FaceCount[GeometryIdx] = FaceCounts[GeometryIdx];
+			GeometryCollection->VertexCount[GeometryIdx] = VertexCounts[GeometryIdx];
+			int32 VertexStart = GeometryCollection->VertexStart[GeometryIdx];
+			int32 VertexEnd = VertexStart + GeometryCollection->VertexCount[GeometryIdx];
+			for (int32 VertexIdx = VertexStart; VertexIdx < VertexEnd; VertexIdx++)
+			{
+				GeometryCollection->BoneMap[VertexIdx] = GeometryCollection->TransformIndex[GeometryIdx];
+			}
 		}
+
+		// remove trailing elements if needed
+		if (NewNumFaces < OldNumFaces)
+		{
+			TArray<int32> ToDelete;
+			AddRange(ToDelete, NewNumFaces, OldNumFaces);
+			GeometryCollection->RemoveElements(FGeometryCollection::FacesGroup, ToDelete);
+		}
+		if (NewNumVertices < OldNumVertices)
+		{
+			TArray<int32> ToDelete;
+			AddRange(ToDelete, NewNumVertices, OldNumVertices);
+			GeometryCollection->RemoveElements(FGeometryCollection::VerticesGroup, ToDelete);
+		}
+
+		// TODO: can remove these ensure()s after this has been tested some more
+		ensure(GeometryCollection->HasContiguousFaces());
+		ensure(GeometryCollection->HasContiguousVertices());
+		ensure(GeometryCollectionAlgo::HasValidGeometryReferences(GeometryCollection));
 	}
 
 	DEFINE_LOG_CATEGORY_STATIC(LogGeometryCollectionClean, Verbose, All);
diff --git a/Engine/Source/Runtime/Experimental/Chaos/Private/GeometryCollection/ManagedArrayCollection.cpp b/Engine/Source/Runtime/Experimental/Chaos/Private/GeometryCollection/ManagedArrayCollection.cpp
--- a/Engine/Source/Runtime/Experimental/Chaos/Private/GeometryCollection/ManagedArrayCollection.cpp
+++ b/Engine/Source/Runtime/Experimental/Chaos/Private/GeometryCollection/ManagedArrayCollection.cpp
@@ -247,12 +247,19 @@ void FManagedArrayCollection::ReorderElements(FName Group, const TArray<int32>&
 	const int32 GroupSize = GroupInfo[Group].Size;
 	check(GroupSize == NewOrder.Num());
 
+	TArray<int32> InverseNewOrder;
+	InverseNewOrder.Init(-1, GroupSize);
+	for (int32 Idx = 0; Idx < GroupSize; ++Idx)
+	{
+		InverseNewOrder[NewOrder[Idx]] = Idx;
+	}
+
 	for (TTuple<FKeyType, FValueType>& Entry : Map)
 	{
 		// Reindex attributes dependent on the group being reordered
 		if (Entry.Value.GroupIndexDependency == Group)
 		{
-			Entry.Value.Value->ReindexFromLookup(NewOrder);
+			Entry.Value.Value->ReindexFromLookup(InverseNewOrder);
 		}
 
 		if (Entry.Key.Get<1>() == Group)
diff --git a/Engine/Source/Runtime/Experimental/Chaos/Public/GeometryCollection/GeometryCollectionAlgo.h b/Engine/Source/Runtime/Experimental/Chaos/Public/GeometryCollection/GeometryCollectionAlgo.h
--- a/Engine/Source/Runtime/Experimental/Chaos/Public/GeometryCollection/GeometryCollectionAlgo.h
+++ b/Engine/Source/Runtime/Experimental/Chaos/Public/GeometryCollection/GeometryCollectionAlgo.h
@@ -182,6 +182,14 @@ namespace GeometryCollectionAlgo
 	CHAOS_API
 	AddFaces(FGeometryCollection* GeometryCollection, const TArray<TArray<FIntVector>> &Faces);
 
+	/**
+	 * Sets new sizes (face/vertex counts) for each geometry
+	 * Added vertices and faces contain arbitrary data; other elements remain unchanged
+	 */
+	void
+	CHAOS_API
+	ResizeGeometries(FGeometryCollection* GeometryCollection, const TArray<int32>& FaceCounts, const TArray<int32>& VertexCounts);
+
 	void
 	CHAOS_API
 	ComputeCoincidentVertices(const FGeometryCollection* GeometryCollection, const float Tolerance, TMap<int32, int32>& CoincidentVerticesMap, TSet<int32>& VertexToDeleteSet);
diff --git a/Engine/Source/Runtime/Experimental/Chaos/Public/GeometryCollection/ManagedArray.h b/Engine/Source/Runtime/Experimental/Chaos/Public/GeometryCollection/ManagedArray.h
--- a/Engine/Source/Runtime/Experimental/Chaos/Public/GeometryCollection/ManagedArray.h
+++ b/Engine/Source/Runtime/Experimental/Chaos/Public/GeometryCollection/ManagedArray.h
@@ -87,14 +87,16 @@ protected:
 
 	/**
 	 * Reorder elements given a new ordering. Sizes must match
+	 * @param NewOrder Mapping from indices in the new array -> indices in the old array
 	 */
 	virtual void Reorder(const TArray<int32>& NewOrder) = 0;
 
 	/** 
-	* Reindex given a lookup table
-	*/
+	 * Reindex given a lookup table
+	 * @param InverseNewOrder Mapping from indices into the old array -> indices in the new array
+	 */
 	//todo: this should really assert, but material is currently relying on both faces and vertices
-	virtual void ReindexFromLookup(const TArray<int32>& NewOrder) { }
+	virtual void ReindexFromLookup(const TArray<int32>& InverseNewOrder) { }
 
 	/**
 	* Init from a predefined Array
@@ -592,7 +594,7 @@ public:
 		}
 	}
 
-	virtual void ReindexFromLookup(const TArray<int32>& NewOrder) override
+	virtual void ReindexFromLookup(const TArray<int32>& InverseNewOrder) override
 	{
 		const int32 ArraySize = Num();
 		for (int32 Index = 0; Index < ArraySize; ++Index)
@@ -600,7 +602,7 @@ public:
 			int32& Mapping = this->operator[](Index);
 			if (Mapping >= 0)
 			{
-				Mapping = NewOrder[Mapping];
+				Mapping = InverseNewOrder[Mapping];
 			}
 		}
 	}
@@ -657,7 +659,7 @@ public:
 		}
 	}
 
-	virtual void ReindexFromLookup(const TArray<int32> & NewOrder) override
+	virtual void ReindexFromLookup(const TArray<int32> & InverseNewOrder) override
 	{
 
 		int32 ArraySize = Num();
@@ -670,7 +672,7 @@ public:
 
 			for (int32 StaleEntry : OldSet)
 			{
-				const int32 NewEntry = StaleEntry >= 0 ? NewOrder[StaleEntry] : StaleEntry;	//only remap if valid
+				const int32 NewEntry = StaleEntry >= 0 ? InverseNewOrder[StaleEntry] : StaleEntry;	//only remap if valid
 				NewSet.Add(NewEntry);
 			}
 		}
@@ -719,7 +721,7 @@ public:
 		}
 	}
 
-	virtual void ReindexFromLookup(const TArray<int32> & NewOrder) override
+	virtual void ReindexFromLookup(const TArray<int32> & InverseNewOrder) override
 	{
 		int32 ArraySize = Num();
 		for (int32 Index = 0; Index < ArraySize; Index++)
@@ -729,7 +731,7 @@ public:
 			{
 				if (RemapVal[i] >= 0)
 				{
-					RemapVal[i] = NewOrder[RemapVal[i]];
+					RemapVal[i] = InverseNewOrder[RemapVal[i]];
 				}
 			}
 		}
