tangents-related fixes for fracture:
- add a "recompute normals" tool to fracture, so artists can fix normals / tangents for existing fractures
- move recompute normals + resample tools to a new category (currently named "fix" -- open to suggestions)
- when cutting with mesh, attempt to use normals / tangents from mesh but recompute if they are invalid (all zero)

diff --git a/Engine/Plugins/Experimental/ChaosEditor/Source/FractureEditor/Private/FractureEditorCommands.h b/Engine/Plugins/Experimental/ChaosEditor/Source/FractureEditor/Private/FractureEditorCommands.h
--- a/Engine/Plugins/Experimental/ChaosEditor/Source/FractureEditor/Private/FractureEditorCommands.h
+++ b/Engine/Plugins/Experimental/ChaosEditor/Source/FractureEditor/Private/FractureEditorCommands.h
@@ -57,9 +57,12 @@ class FFractureEditorCommands : public TCommands<FFractureEditorCommands>
 		TSharedPtr< FUICommandInfo > Slice;
 		TSharedPtr< FUICommandInfo > Brick;
 		TSharedPtr< FUICommandInfo > Texture;
-		TSharedPtr< FUICommandInfo > Resample;
 		TSharedPtr< FUICommandInfo > Mesh;
 
+		// Cleanup Commands
+		TSharedPtr< FUICommandInfo > RecomputeNormals;
+		TSharedPtr< FUICommandInfo > Resample;
+
 		// Property Commands
 		TSharedPtr< FUICommandInfo > SetInitialDynamicState;
 		
diff --git a/Engine/Plugins/Experimental/ChaosEditor/Source/FractureEditor/Private/FractureEditorModeToolkit.cpp b/Engine/Plugins/Experimental/ChaosEditor/Source/FractureEditor/Private/FractureEditorModeToolkit.cpp
--- a/Engine/Plugins/Experimental/ChaosEditor/Source/FractureEditor/Private/FractureEditorModeToolkit.cpp
+++ b/Engine/Plugins/Experimental/ChaosEditor/Source/FractureEditor/Private/FractureEditorModeToolkit.cpp
@@ -463,7 +463,7 @@ void FFractureEditorModeToolkit::OnObjectPostEditChange( UObject* Object, FPrope
 	}
 }
 
-const TArray<FName> FFractureEditorModeToolkit::PaletteNames = { FName(TEXT("Generate")), FName(TEXT("Select")), FName(TEXT("Fracture")), FName(TEXT("Edit")), FName(TEXT("Cluster")), FName(TEXT("Embed")), FName(TEXT("UV")), FName(TEXT("Properties")) };
+const TArray<FName> FFractureEditorModeToolkit::PaletteNames = { FName(TEXT("Generate")), FName(TEXT("Select")), FName(TEXT("Fracture")), FName(TEXT("Edit")), FName(TEXT("Cluster")), FName(TEXT("Embed")), FName(TEXT("Fix")), FName(TEXT("UV")), FName(TEXT("Properties")) };
 
 FText FFractureEditorModeToolkit::GetToolPaletteDisplayName(FName Palette) const
 { 
@@ -510,7 +510,6 @@ void FFractureEditorModeToolkit::BuildToolPalette(FName PaletteIndex, class FToo
 		ToolbarBuilder.AddToolBarButton(Commands.Mesh);
 		ToolbarBuilder.AddToolBarButton(Commands.Slice);
 		ToolbarBuilder.AddToolBarButton(Commands.Brick);
-		ToolbarBuilder.AddToolBarButton(Commands.Resample);
 	}
 	else if (PaletteIndex == TEXT("Edit"))
 	{
@@ -534,6 +533,11 @@ void FFractureEditorModeToolkit::BuildToolPalette(FName PaletteIndex, class FToo
 	{
 		ToolbarBuilder.AddToolBarButton(Commands.AutoUV);
 	}
+	else if (PaletteIndex == TEXT("Fix"))
+	{
+		ToolbarBuilder.AddToolBarButton(Commands.RecomputeNormals);
+		ToolbarBuilder.AddToolBarButton(Commands.Resample);
+	}
 	else if (PaletteIndex == TEXT("Properties"))
 	{
 		ToolbarBuilder.AddToolBarButton(Commands.SetInitialDynamicState);
diff --git a/Engine/Plugins/Experimental/ChaosEditor/Source/FractureEditor/Private/FractureToolRecomputeNormals.cpp b/Engine/Plugins/Experimental/ChaosEditor/Source/FractureEditor/Private/FractureToolRecomputeNormals.cpp
new file mode 100644
--- /dev/null
+++ b/Engine/Plugins/Experimental/ChaosEditor/Source/FractureEditor/Private/FractureToolRecomputeNormals.cpp
@@ -0,0 +1,150 @@
+// Copyright Epic Games, Inc. All Rights Reserved.
+
+#include "FractureToolRecomputeNormals.h"
+#include "FractureEditorStyle.h"
+#include "PlanarCut.h"
+#include "GeometryCollection/GeometryCollection.h"
+#include "GeometryCollection/GeometryCollectionObject.h"
+#include "GeometryCollection/GeometryCollectionAlgo.h"
+#include "FractureEditorModeToolkit.h"
+#include "FractureToolContext.h"
+
+#define LOCTEXT_NAMESPACE "FractureRecomputeNormals"
+
+
+UFractureToolRecomputeNormals::UFractureToolRecomputeNormals(const FObjectInitializer& ObjInit) 
+	: Super(ObjInit) 
+{
+	NormalsSettings = NewObject<UFractureRecomputeNormalsSettings>(GetTransientPackage(), UFractureRecomputeNormalsSettings::StaticClass());
+	NormalsSettings->OwnerTool = this;
+}
+
+FText UFractureToolRecomputeNormals::GetDisplayText() const
+{
+	return FText(LOCTEXT("FractureToolRecomputeNormals", "Recompute Normals and Tangents"));
+}
+
+FText UFractureToolRecomputeNormals::GetTooltipText() const 
+{
+	return FText(LOCTEXT("FractureToolRecomputeNormalsTooltip", "The Recompute Normals tool recompute normals and tangents for selected geometry."));
+}
+
+FSlateIcon UFractureToolRecomputeNormals::GetToolIcon() const 
+{
+	return FSlateIcon("FractureEditorStyle", "FractureEditor.RecomputeNormals");
+}
+
+void UFractureToolRecomputeNormals::RegisterUICommand( FFractureEditorCommands* BindingContext ) 
+{
+	UI_COMMAND_EXT( BindingContext, UICommandInfo, "RecomputeNormals", "Normals", "Recompute Normals and Tangents", EUserInterfaceActionType::ToggleButton, FInputChord() );
+	BindingContext->RecomputeNormals = UICommandInfo;
+}
+
+void UFractureToolRecomputeNormals::Render(const FSceneView* View, FViewport* Viewport, FPrimitiveDrawInterface* PDI)
+{
+	if (NormalsSettings->bShowTangents || NormalsSettings->bShowNormals)
+	{
+		for (int32 Idx = 0; Idx < DisplayVertices.Num(); Idx++)
+		{
+			FVector Point = DisplayVertices[Idx];
+			FVector Normal = DisplayNormals[Idx] * NormalsSettings->Length;
+			FVector TanU = DisplayTanUs[Idx] * NormalsSettings->Length;
+			FVector TanV = DisplayTanVs[Idx] * NormalsSettings->Length;
+			if (NormalsSettings->bShowNormals)
+			{
+				PDI->DrawLine(Point, Point + Normal, FLinearColor::Red, SDPG_Foreground);
+			}
+			if (NormalsSettings->bShowTangents)
+			{
+				PDI->DrawLine(Point, Point + TanU, FLinearColor::Green, SDPG_Foreground);
+				PDI->DrawLine(Point, Point + TanV, FLinearColor::Blue, SDPG_Foreground);
+			}
+		}
+	}
+}
+
+TArray<UObject*> UFractureToolRecomputeNormals::GetSettingsObjects() const
+ {
+	TArray<UObject*> Settings;
+	Settings.Add(NormalsSettings);
+	return Settings;
+}
+
+
+void UFractureToolRecomputeNormals::FractureContextChanged()
+{
+	TArray<FFractureToolContext> FractureContexts = GetFractureToolContexts();
+
+	DisplayVertices.Reset();
+	DisplayNormals.Reset();
+	DisplayTanUs.Reset();
+	DisplayTanVs.Reset();
+	for (const FFractureToolContext& FractureContext : FractureContexts)
+	{
+		FGeometryCollection& Collection = *FractureContext.GetGeometryCollection();
+		
+
+		FTransform OuterTransform = FractureContext.GetTransform();
+		for (int32 TransformIdx : FractureContext.GetSelection())
+		{
+			FTransform InnerTransform = GeometryCollectionAlgo::GlobalMatrix(Collection.Transform, Collection.Parent, TransformIdx);
+			if (Collection.HasAttribute("ExplodedVector", FGeometryCollection::TransformGroup))
+			{
+				TManagedArray<FVector>& ExplodedVectors = Collection.GetAttribute<FVector>("ExplodedVector", FGeometryCollection::TransformGroup);
+				InnerTransform = InnerTransform * FTransform(ExplodedVectors[TransformIdx]);
+			}
+			
+			FTransform CombinedTransform = InnerTransform * OuterTransform;
+			int32 GeometryIdx = Collection.TransformToGeometryIndex[TransformIdx];
+			if (GeometryIdx > -1)
+			{
+				TArray<bool> ShowVerts;
+				int32 VertStart = Collection.VertexStart[GeometryIdx];
+				int32 VertCount = Collection.VertexCount[GeometryIdx];
+				int32 FaceStart = Collection.FaceStart[GeometryIdx];
+				int32 FaceCount = Collection.FaceCount[GeometryIdx];
+				bool bShowAll = !NormalsSettings->bOnlyInternalSurfaces;
+				ShowVerts.Init(bShowAll, VertCount);
+				if (!bShowAll)
+				{
+					for (int FIdx = FaceStart; FIdx < FaceStart + FaceCount; FIdx++)
+					{
+						if (Collection.Visible[FIdx] && (Collection.MaterialID[FIdx] % 2) == 1)
+						{
+							FIntVector Face = Collection.Indices[FIdx];
+							ShowVerts[Face.X - VertStart] = true;
+							ShowVerts[Face.Y - VertStart] = true;
+							ShowVerts[Face.Z - VertStart] = true;
+						}
+					}
+				}
+				for (int32 VIdx = VertStart; VIdx < VertStart + VertCount; VIdx++)
+				{
+					if (ShowVerts[VIdx - VertStart])
+					{
+						DisplayVertices.Add(CombinedTransform.TransformPosition(Collection.Vertex[VIdx]));
+						DisplayNormals.Add(CombinedTransform.TransformVectorNoScale(Collection.Normal[VIdx]));
+						DisplayTanUs.Add(CombinedTransform.TransformVectorNoScale(Collection.TangentU[VIdx]));
+						DisplayTanVs.Add(CombinedTransform.TransformVectorNoScale(Collection.TangentV[VIdx]));
+					}
+				}
+			}
+		}
+	}
+}
+
+int32 UFractureToolRecomputeNormals::ExecuteFracture(const FFractureToolContext& FractureContext)
+{
+	if (FractureContext.IsValid())
+	{
+		const UFractureCollisionSettings* LocalCutSettings = CollisionSettings;
+
+		bool bOnlyInternalFaces = true;
+		RecomputeNormalsAndTangents(NormalsSettings->bOnlyTangents, *FractureContext.GetGeometryCollection(), FractureContext.GetSelection(), NormalsSettings->bOnlyInternalSurfaces);
+	}
+
+	return INDEX_NONE;
+}
+
+
+#undef LOCTEXT_NAMESPACE
\ No newline at end of file
diff --git a/Engine/Plugins/Experimental/ChaosEditor/Source/FractureEditor/Private/FractureToolRecomputeNormals.h b/Engine/Plugins/Experimental/ChaosEditor/Source/FractureEditor/Private/FractureToolRecomputeNormals.h
new file mode 100644
--- /dev/null
+++ b/Engine/Plugins/Experimental/ChaosEditor/Source/FractureEditor/Private/FractureToolRecomputeNormals.h
@@ -0,0 +1,82 @@
+// Copyright Epic Games, Inc. All Rights Reserved.
+
+#pragma once
+
+#include "FractureToolCutter.h"
+#include "Engine/StaticMeshActor.h"
+
+#include "FractureToolRecomputeNormals.generated.h"
+
+class FFractureToolContext;
+
+/** Settings for visualizing and recomputing normals and tangents */
+UCLASS(config = EditorPerProjectUserSettings)
+class UFractureRecomputeNormalsSettings : public UFractureToolSettings
+{
+public:
+
+	GENERATED_BODY()
+
+	UFractureRecomputeNormalsSettings(const FObjectInitializer& ObjInit)
+	: Super(ObjInit)
+	{}
+
+	/** Whether to display normal vectors */
+	UPROPERTY(EditAnywhere, Category = DisplaySettings)
+	bool bShowNormals = true;
+
+	/** Whether to display tangent and bitangent vectors */
+	UPROPERTY(EditAnywhere, Category = DisplaySettings)
+	bool bShowTangents = true;
+
+	/** Length of display normal and tangent vectors */
+	UPROPERTY(EditAnywhere, Category = DisplaySettings, meta = (ClampMin = ".001", ClampMax = "100", UIMax = "10"))
+	float Length = 2.0f;
+
+	/** Whether to only recompute tangents, and leave normals as they were */
+	UPROPERTY(EditAnywhere, Category = RecomputeSettings)
+	bool bOnlyTangents = false;
+
+	/** Whether to only change internal surface normals / tangents */
+	UPROPERTY(EditAnywhere, Category = RecomputeSettings, AdvancedDisplay)
+	bool bOnlyInternalSurfaces = true;
+};
+
+
+// Note this tool doesn't actually fracture, but it does remake pieces of geometry and shares a lot of machinery with the fracture tools
+UCLASS(DisplayName = "Recompute Normals Tool", Category = "FractureTools")
+class UFractureToolRecomputeNormals : public UFractureToolCutterBase
+{
+public:
+	GENERATED_BODY()
+
+	UFractureToolRecomputeNormals(const FObjectInitializer& ObjInit);
+
+	// UFractureTool Interface
+	virtual FText GetDisplayText() const override;
+	virtual FText GetTooltipText() const override;
+	virtual FSlateIcon GetToolIcon() const override;
+
+	virtual FText GetApplyText() const override { return FText(NSLOCTEXT("FractureRecomputeNormals", "ExecuteRecomputeNormals", "Recompute")); }
+
+	void Render(const FSceneView* View, FViewport* Viewport, FPrimitiveDrawInterface* PDI) override;
+
+	virtual void RegisterUICommand(FFractureEditorCommands* BindingContext) override;
+
+	virtual TArray<UObject*> GetSettingsObjects() const override;
+
+	virtual void FractureContextChanged() override;
+	virtual int32 ExecuteFracture(const FFractureToolContext& FractureContext) override;
+
+protected:
+	UPROPERTY()
+	TObjectPtr<UFractureRecomputeNormalsSettings> NormalsSettings;
+
+	TArray<FVector> DisplayVertices;
+	TArray<FVector> DisplayNormals;
+	TArray<FVector> DisplayTanUs;
+	TArray<FVector> DisplayTanVs;
+
+};
+
+
diff --git a/Engine/Plugins/Experimental/GeometryProcessing/Source/DynamicMesh/Private/MeshTangents.cpp b/Engine/Plugins/Experimental/GeometryProcessing/Source/DynamicMesh/Private/MeshTangents.cpp
--- a/Engine/Plugins/Experimental/GeometryProcessing/Source/DynamicMesh/Private/MeshTangents.cpp
+++ b/Engine/Plugins/Experimental/GeometryProcessing/Source/DynamicMesh/Private/MeshTangents.cpp
@@ -151,6 +151,47 @@ void TMeshTangents<RealType>::ComputeTriangleTangents(const FDynamicMeshUVOverla
 
 
 
+template<typename RealType>
+bool TMeshTangents<RealType>::CopyToOverlays(FDynamicMesh3& MeshToSet)
+{
+	if (!MeshToSet.HasAttributes() || MeshToSet.Attributes()->NumNormalLayers() != 3)
+	{
+		return false;
+	}
+
+	// Set aliases to make iterating over tangents and bitangents easier
+	FDynamicMeshNormalOverlay* TangentOverlays[2] = { MeshToSet.Attributes()->PrimaryTangents(), MeshToSet.Attributes()->PrimaryBiTangents() };
+	TArray<FVector3<RealType>>* TangentValues[2] = { &Tangents, &Bitangents };
+	
+	for (int Idx = 0; Idx < 2; Idx++)
+	{
+		// Create overlay topology
+		TArray<FVector3<RealType>>& TV = *TangentValues[Idx];
+		TangentOverlays[Idx]->CreateFromPredicate([&MeshToSet, &TV](int ParentVertexIdx, int TriIDA, int TriIDB) -> bool
+		{
+			FIndex3i TriA = MeshToSet.GetTriangle(TriIDA);
+			FIndex3i TriB = MeshToSet.GetTriangle(TriIDB);
+			int SubA = TriA.IndexOf(ParentVertexIdx);
+			int SubB = TriB.IndexOf(ParentVertexIdx);
+			checkSlow(SubA > -1 && SubB > -1);
+			return TV[TriIDA * 3 + SubA].DistanceSquared(TV[TriIDB * 3 + SubB]) < TMathUtil<RealType>::ZeroTolerance;
+		}, 0);
+
+		// Write tangent values out for each wedge value
+		// Note: shared elements will be written to multiple times, and the last value written will be used
+		for (int TID : MeshToSet.TriangleIndicesItr())
+		{
+			FIndex3i ElTri = TangentOverlays[Idx]->GetTriangle(TID);
+			for (int SubIdx = 0; SubIdx < 3; SubIdx++)
+			{
+				TangentOverlays[Idx]->SetElement(ElTri[SubIdx], (FVector3f)TV[TID * 3 + SubIdx]);
+			}
+		}
+	}
+	return true;
+}
+
+
 
 static FVector3d PlaneProjectionNormalized(const FVector3d& Vector, const FVector3d& PlaneNormal)
 {
diff --git a/Engine/Plugins/Experimental/GeometryProcessing/Source/DynamicMesh/Public/MeshTangents.h b/Engine/Plugins/Experimental/GeometryProcessing/Source/DynamicMesh/Public/MeshTangents.h
--- a/Engine/Plugins/Experimental/GeometryProcessing/Source/DynamicMesh/Public/MeshTangents.h
+++ b/Engine/Plugins/Experimental/GeometryProcessing/Source/DynamicMesh/Public/MeshTangents.h
@@ -178,6 +178,11 @@ public:
 	void ComputeTriangleTangents(const FDynamicMeshUVOverlay* UVOverlay, bool bOrthogonalize = true);
 
 
+	/**
+	 * Set Tangents on mesh overlays
+	 * @param MeshToSet Mesh to copy overlays to; does not need to be the same as the Mesh member of this class
+	 */
+	bool CopyToOverlays(FDynamicMesh3& MeshToSet);
 
 protected:
 	/**
diff --git a/Engine/Plugins/Experimental/PlanarCutPlugin/Source/PlanarCut/PlanarCut.Build.cs b/Engine/Plugins/Experimental/PlanarCutPlugin/Source/PlanarCut/PlanarCut.Build.cs
--- a/Engine/Plugins/Experimental/PlanarCutPlugin/Source/PlanarCut/PlanarCut.Build.cs
+++ b/Engine/Plugins/Experimental/PlanarCutPlugin/Source/PlanarCut/PlanarCut.Build.cs
@@ -8,7 +8,14 @@ namespace UnrealBuildTool.Rules
 		{
 			PrivateIncludePaths.Add("PlanarCut/Private");
             PublicIncludePaths.Add(ModuleDirectory + "/Public");
-            PrivateDependencyModuleNames.Add("MeshUtilitiesCommon");
+            PrivateDependencyModuleNames.AddRange(
+				new string[]
+				{
+					"MeshUtilitiesCommon",
+					"StaticMeshDescription",
+					"MeshConversion",
+				}
+				);
 			PublicDependencyModuleNames.AddRange(
 				new string[]
 				{
@@ -20,7 +27,6 @@ namespace UnrealBuildTool.Rules
 					"Chaos",
 					"GeometryAlgorithms",
 					"MeshDescription",
-					"MeshConversion",
                 }
                 );
 		}
diff --git a/Engine/Plugins/Experimental/PlanarCutPlugin/Source/PlanarCut/Private/PlanarCut.cpp b/Engine/Plugins/Experimental/PlanarCutPlugin/Source/PlanarCut/Private/PlanarCut.cpp
--- a/Engine/Plugins/Experimental/PlanarCutPlugin/Source/PlanarCut/Private/PlanarCut.cpp
+++ b/Engine/Plugins/Experimental/PlanarCutPlugin/Source/PlanarCut/Private/PlanarCut.cpp
@@ -29,8 +29,11 @@
 #include "QueueRemesher.h"
 #include "DynamicVertexAttribute.h"
 #include "MeshNormals.h"
+#include "MeshTangents.h"
 #include "ConstrainedDelaunay2.h"
 
+#include "Engine/EngineTypes.h"
+#include "StaticMeshOperations.h"
 #include "MeshDescriptionToDynamicMesh.h"
 
 #include "Algo/Rotate.h"
@@ -148,6 +151,68 @@ namespace UE
 				Vs->GetValue(VID, V);
 			}
 
+			void ComputeTangents(FDynamicMesh3& Mesh, bool bOnlyOddMaterials, const TArrayView<const int32>& WhichMaterials, bool bRecomputeNormals = true)
+			{
+				FDynamicMeshNormalOverlay* Normals = Mesh.Attributes()->PrimaryNormals();
+				FMeshNormals::InitializeOverlayToPerVertexNormals(Normals, !bRecomputeNormals);
+				if (bRecomputeNormals)
+				{
+					FMeshNormals::QuickRecomputeOverlayNormals(Mesh);
+				}
+
+				// Copy per-vertex UVs to a UV overlay, because that's what the tangents code uses
+				// (TODO: consider making a tangent computation path that uses vertex normals / UVs)
+				FDynamicMeshUVOverlay* UVs = Mesh.Attributes()->PrimaryUV();
+				UVs->ClearElements();
+				TArray<int> VertToUVMap;
+				VertToUVMap.SetNumUninitialized(Mesh.MaxVertexID());
+				for (int VID : Mesh.VertexIndicesItr())
+				{
+					FVector2f UV = Mesh.GetVertexUV(VID);
+					int UVID = UVs->AppendElement(UV);
+					VertToUVMap[VID] = UVID;
+				}
+				for (int TID : Mesh.TriangleIndicesItr())
+				{
+					FIndex3i Tri = Mesh.GetTriangle(TID);
+					Tri.A = VertToUVMap[Tri.A];
+					Tri.B = VertToUVMap[Tri.B];
+					Tri.C = VertToUVMap[Tri.C];
+					UVs->SetTriangle(TID, Tri);
+				}
+
+				FComputeTangentsOptions Options;
+				Options.bAngleWeighted = true;
+				Options.bAveraged = true;
+				FMeshTangentsf Tangents(&Mesh);
+				Tangents.ComputeTriVertexTangents(Normals, UVs, Options);
+
+				const TArray<FVector3f>& TanU = Tangents.GetTangents();
+				const TArray<FVector3f>& TanV = Tangents.GetBitangents();
+				FDynamicMeshMaterialAttribute* MaterialIDs = Mesh.Attributes()->GetMaterialID();
+				for (int TID : Mesh.TriangleIndicesItr())
+				{
+					int MaterialID = MaterialIDs->GetValue(TID);
+					if (bOnlyOddMaterials && (MaterialID % 2) == 0)
+					{
+						continue;
+					}
+					else if (WhichMaterials.Contains(MaterialID))
+					{
+						continue;
+					}
+					
+					int TanIdxBase = TID * 3;
+					FIndex3i Tri = Mesh.GetTriangle(TID);
+					for (int Idx = 0; Idx < 3; Idx++)
+					{
+						int VID = Tri[Idx];
+						int TanIdx = TanIdxBase + Idx;
+						UE::PlanarCutInternals::AugmentDynamicMesh::SetTangent(Mesh, VID, Mesh.GetVertexNormal(VID), TanU[TanIdx], TanV[TanIdx]);
+					}
+				}
+			}
+
 			// per component sampling is a rough heuristic to avoid doing geodesic distance but still get points on a 'thin' slice
 			void AddCollisionSamplesPerComponent(FDynamicMesh3& Mesh, double Spacing)
 			{
@@ -1255,12 +1320,12 @@ struct FDynamicMeshCollection
 	TIndirectArray<FMeshData> Meshes;
 	FAxisAlignedBox3d Bounds;
 
-	FDynamicMeshCollection(const FGeometryCollection* Collection, const TArrayView<const int32>& TransformIndices, FTransform TransformCollection)
+	FDynamicMeshCollection(const FGeometryCollection* Collection, const TArrayView<const int32>& TransformIndices, FTransform TransformCollection, bool bSaveIsolatedVertices = false)
 	{
-		Init(Collection, TransformIndices, TransformCollection);
+		Init(Collection, TransformIndices, TransformCollection, bSaveIsolatedVertices);
 	}
 
-	void Init(const FGeometryCollection* Collection, const TArrayView<const int32>& TransformIndices, FTransform TransformCollection)
+	void Init(const FGeometryCollection* Collection, const TArrayView<const int32>& TransformIndices, FTransform TransformCollection, bool bSaveIsolatedVertices = false)
 	{
 		Meshes.Reset();
 		Bounds = FAxisAlignedBox3d::Empty();
@@ -1334,8 +1399,11 @@ struct FDynamicMeshCollection
 				// note: material index doesn't need to be passed through; will be rebuilt by a call to reindex materials once the cut mesh is returned back to geometry collection format
 			}
 
-			FDynamicMeshEditor Editor(&Mesh);
-			Editor.RemoveIsolatedVertices();
+			if (!bSaveIsolatedVertices)
+			{
+				FDynamicMeshEditor Editor(&Mesh);
+				Editor.RemoveIsolatedVertices();
+			}
 
 			Bounds.Contain(Mesh.GetCachedBounds());
 
@@ -1995,6 +2063,42 @@ struct FDynamicMeshCollection
 		}
 	}
 
+	// Update all geometry in a GeometryCollection w/ the meshes in the MeshCollection
+	// Resizes the GeometryCollection as needed
+	bool UpdateAllCollections(FGeometryCollection& Collection)
+	{
+		bool bAllSucceeded = true;
+
+		int32 NumGeometry = Collection.NumElements(FGeometryCollection::GeometryGroup);
+		TArray<int32> NewFaceCounts, NewVertexCounts;
+		NewFaceCounts.SetNumUninitialized(NumGeometry);
+		NewVertexCounts.SetNumUninitialized(NumGeometry);
+		for (int32 GeomIdx = 0; GeomIdx < Collection.FaceCount.Num(); GeomIdx++)
+		{
+			NewFaceCounts[GeomIdx] = Collection.FaceCount[GeomIdx];
+			NewVertexCounts[GeomIdx] = Collection.VertexCount[GeomIdx];
+		}
+		for (int MeshIdx = 0; MeshIdx < Meshes.Num(); MeshIdx++)
+		{
+			FDynamicMeshCollection::FMeshData& MeshData = Meshes[MeshIdx];
+			int32 GeomIdx = Collection.TransformToGeometryIndex[MeshData.TransformIndex];
+			NewFaceCounts[GeomIdx] = MeshData.AugMesh.TriangleCount();
+			NewVertexCounts[GeomIdx] = MeshData.AugMesh.VertexCount();
+		}
+		GeometryCollectionAlgo::ResizeGeometries(&Collection, NewFaceCounts, NewVertexCounts);
+
+		for (int MeshIdx = 0; MeshIdx < Meshes.Num(); MeshIdx++)
+		{
+			FDynamicMeshCollection::FMeshData& MeshData = Meshes[MeshIdx];
+			FDynamicMesh3& Mesh = MeshData.AugMesh;
+			int32 GeometryIdx = Collection.TransformToGeometryIndex[MeshData.TransformIndex];
+			bool bSucceeded = UpdateCollection(MeshData.ToCollection, Mesh, GeometryIdx, Collection, -1);
+			bAllSucceeded &= bSucceeded;
+		}
+
+		return bAllSucceeded;
+	}
+
 	// Update an existing geometry in a collection w/ a new mesh (w/ the same number of faces and vertices!)
 	static bool UpdateCollection(const FTransform& ToCollection, FDynamicMesh3& Mesh, int32 GeometryIdx, FGeometryCollection& Output, int32 InternalMaterialID)
 	{
@@ -2741,6 +2845,21 @@ int32 CutWithMesh(
 	FMeshDescriptionToDynamicMesh Converter;
 	FDynamicMesh3 FullMesh; // full-featured conversion of the source mesh
 	Converter.Convert(CuttingMesh, FullMesh, true);
+	bool bHasInvalidNormals, bHasInvalidTangents;
+	FStaticMeshOperations::AreNormalsAndTangentsValid(*CuttingMesh, bHasInvalidNormals, bHasInvalidTangents);
+	if (bHasInvalidNormals || bHasInvalidTangents)
+	{
+		FDynamicMeshAttributeSet& Attribs = *FullMesh.Attributes();
+		FDynamicMeshNormalOverlay* NTB[3]{ Attribs.PrimaryNormals(), Attribs.PrimaryTangents(), Attribs.PrimaryBiTangents() };
+		if (bHasInvalidNormals)
+		{
+			FMeshNormals::InitializeOverlayToPerVertexNormals(NTB[0], false);
+		}
+		FMeshTangentsf Tangents(&FullMesh);
+		Tangents.ComputeTriVertexTangents(NTB[0], Attribs.PrimaryUV(), { true, true });
+		Tangents.CopyToOverlays(FullMesh);
+	}
+
 	FDynamicMesh3 DynamicCuttingMesh; // version of mesh that is split apart at seams to be compatible w/ geometry collection, with corresponding attributes set
 	UE::PlanarCutInternals::AugmentDynamicMesh::Augment(DynamicCuttingMesh);
 	// Note: This conversion will likely go away, b/c I plan to switch over to doing the boolean operations on the fuller rep, but the code can be adapted
@@ -2836,6 +2955,25 @@ int32 CutWithMesh(
 }
 
 
+void RecomputeNormalsAndTangents(bool bOnlyTangents, FGeometryCollection& Collection, const TArrayView<const int32>& TransformIndices,
+								 bool bOnlyOddMaterials, const TArrayView<const int32>& WhichMaterials)
+{
+	FTransform CellsToWorld = FTransform::Identity;
+
+	FDynamicMeshCollection MeshCollection(&Collection, TransformIndices, CellsToWorld, true);
+
+	for (int MeshIdx = 0; MeshIdx < MeshCollection.Meshes.Num(); MeshIdx++)
+	{
+		FDynamicMesh3& Mesh = MeshCollection.Meshes[MeshIdx].AugMesh;
+		UE::PlanarCutInternals::AugmentDynamicMesh::ComputeTangents(Mesh, bOnlyOddMaterials, WhichMaterials, !bOnlyTangents);
+	}
+
+	MeshCollection.UpdateAllCollections(Collection);
+
+	Collection.ReindexMaterials();
+}
+
+
 
 int32 AddCollisionSampleVertices(double CollisionSampleSpacing, FGeometryCollection& Collection, const TArrayView<const int32>& TransformIndices)
 {
@@ -2845,31 +2983,7 @@ int32 AddCollisionSampleVertices(double CollisionSampleSpacing, FGeometryCollect
 
 	MeshCollection.AddCollisionSamples(CollisionSampleSpacing);
 
-	int32 NumGeometry = Collection.NumElements(FGeometryCollection::GeometryGroup);
-	TArray<int32> NewFaceCounts, NewVertexCounts;
-	NewFaceCounts.SetNumUninitialized(NumGeometry);
-	NewVertexCounts.SetNumUninitialized(NumGeometry);
-	for (int32 GeomIdx = 0; GeomIdx < Collection.FaceCount.Num(); GeomIdx++)
-	{
-		NewFaceCounts[GeomIdx] = Collection.FaceCount[GeomIdx];
-		NewVertexCounts[GeomIdx] = Collection.VertexCount[GeomIdx];
-	}
-	for (int MeshIdx = 0; MeshIdx < MeshCollection.Meshes.Num(); MeshIdx++)
-	{
-		FDynamicMeshCollection::FMeshData& MeshData = MeshCollection.Meshes[MeshIdx];
-		int32 GeomIdx = Collection.TransformToGeometryIndex[MeshData.TransformIndex];
-		NewFaceCounts[GeomIdx] = MeshData.AugMesh.TriangleCount();
-		NewVertexCounts[GeomIdx] = MeshData.AugMesh.VertexCount();
-	}
-	GeometryCollectionAlgo::ResizeGeometries(&Collection, NewFaceCounts, NewVertexCounts);
-
-	for (int MeshIdx = 0; MeshIdx < MeshCollection.Meshes.Num(); MeshIdx++)
-	{
-		FDynamicMeshCollection::FMeshData& MeshData = MeshCollection.Meshes[MeshIdx];
-		FDynamicMesh3& Mesh = MeshData.AugMesh;
-		int32 GeometryIdx = Collection.TransformToGeometryIndex[MeshData.TransformIndex];
-		ensure(MeshCollection.UpdateCollection(MeshData.ToCollection, Mesh, GeometryIdx, Collection, -1));
-	}
+	MeshCollection.UpdateAllCollections(Collection);
 
 	Collection.ReindexMaterials();
 
diff --git a/Engine/Plugins/Experimental/PlanarCutPlugin/Source/PlanarCut/Public/PlanarCut.h b/Engine/Plugins/Experimental/PlanarCutPlugin/Source/PlanarCut/Public/PlanarCut.h
--- a/Engine/Plugins/Experimental/PlanarCutPlugin/Source/PlanarCut/Public/PlanarCut.h
+++ b/Engine/Plugins/Experimental/PlanarCutPlugin/Source/PlanarCut/Public/PlanarCut.h
@@ -208,6 +208,17 @@ int32 PLANARCUT_API CutMultipleWithMultiplePlanes(
 	bool bSetDefaultInternalMaterialsFromCollection = true
 );
 
+/**
+ * Recompute normals and tangents of selected geometry, optionally restricted to faces with odd or given material IDs (i.e. to target internal faces)
+ *
+ * @param bOnlyTangents		If true, leave normals unchanged and only recompute tangent&bitangent vectors
+ * @param Collection		The Geometry Collection to be updated
+ * @param TransformIndices	Which transform groups on the Geometry Collection to be updated.  If empty, all groups are updated.
+ * @param bOnlyOddMaterials	If true, restrict recomputation to odd-numbered material IDs
+ * @param WhichMaterials	If non-empty, restrict recomputation to only the listed material IDs
+ */
+void PLANARCUT_API RecomputeNormalsAndTangents(bool bOnlyTangents, FGeometryCollection& Collection, const TArrayView<const int32>& TransformIndices = TArrayView<const int32>(),
+	bool bOnlyOddMaterials = true, const TArrayView<const int32>& WhichMaterials = TArrayView<const int32>());
 
 /**
  * Scatter additional vertices (w/ no associated triangle) as needed to satisfy minimum point spacing
