Make UV packer robust to zero-area islands + make fracture AutoUV automatically fix islands with collapsed UVs
#rnx

diff --git a/Engine/Plugins/Experimental/GeometryProcessing/Source/GeometricObjects/Private/Parameterization/UVPacking.cpp b/Engine/Plugins/Experimental/GeometryProcessing/Source/GeometricObjects/Private/Parameterization/UVPacking.cpp
--- a/Engine/Plugins/Experimental/GeometryProcessing/Source/GeometricObjects/Private/Parameterization/UVPacking.cpp
+++ b/Engine/Plugins/Experimental/GeometryProcessing/Source/GeometricObjects/Private/Parameterization/UVPacking.cpp
@@ -41,7 +41,7 @@ struct FUVIsland
 	FVector2d	MaxUV;
 
 	double		UVArea;
-	FVector2d	WorldScale;
+	double		ScaleToWorld; // Scale factor that would make 1 texel ~= 1 unit in world space
 
 	FVector2d	UVScale;
 	FVector2d	PackingScaleU;
@@ -100,7 +100,7 @@ bool FStandardChartPacker::FindBestPacking(TArray<FUVIsland>& Charts)
 	TotalUVArea = 0.0;
 	for (const FUVIsland& Chart : Charts)
 	{
-		TotalUVArea += Chart.UVArea * Chart.WorldScale.X * Chart.WorldScale.Y;
+		TotalUVArea += Chart.UVArea * Chart.ScaleToWorld * Chart.ScaleToWorld;
 	}
 
 	if (TotalUVArea <= 0.0)
@@ -255,7 +255,7 @@ void FStandardChartPacker::ScaleCharts( TArray<FUVIsland>& Charts, double UVScal
 	for( int32 i = 0; i <Charts.Num(); i++ )
 	{
 		FUVIsland& Chart = Charts[i];
-		Chart.UVScale = Chart.WorldScale * UVScale;
+		Chart.UVScale = FVector2d::One() * UVScale * Chart.ScaleToWorld;
 	}
 	
 	// Unsort the charts to make sure ScaleCharts always return the same ordering
@@ -322,7 +322,7 @@ void FStandardChartPacker::ScaleCharts( TArray<FUVIsland>& Charts, double UVScal
 	}
 #endif
 
-#if 1
+#if 0 // TODO: re-enable under a boolean flag if we find a case where we want to allow non-uniform scaling
 	double NonuniformScale = 1.0f;
 	for (int i = 0; i < 1000; i++)
 	{
@@ -825,7 +825,7 @@ bool FUVPacker::StandardPack(IUVMeshView* Mesh, int NumIslands, TFunctionRef<voi
 		Chart.MinUV = FVector2d(FLT_MAX, FLT_MAX);
 		Chart.MaxUV = FVector2d(-FLT_MAX, -FLT_MAX);
 		Chart.UVArea = 0.0f;
-		Chart.WorldScale = FVector2d::Zero();
+		double UVLengthSum = 0, WorldLengthSum = 0;
 
 		for (int32 tid : Chart.Triangles)
 		{
@@ -848,21 +848,28 @@ bool FUVPacker::StandardPack(IUVMeshView* Mesh, int NumIslands, TFunctionRef<voi
 
 			FVector3d Edge1 = Positions[1] - Positions[0];
 			FVector3d Edge2 = Positions[2] - Positions[0];
-			double Area = 0.5f * (Edge1.Cross(Edge2)).Length();
+			FVector3d Edge3 = Positions[2] - Positions[1];
+			double WorldLength = Edge1.Length() + Edge2.Length() + Edge3.Length();
 
 			FVector2d EdgeUV1 = UVs[1] - UVs[0];
 			FVector2d EdgeUV2 = UVs[2] - UVs[0];
+			FVector2d EdgeUV3 = UVs[2] - UVs[1];
+			double UVLength = EdgeUV1.Length() + EdgeUV2.Length() + EdgeUV3.Length();
 			double UVArea = 0.5f * FMathd::Abs(EdgeUV1.X * EdgeUV2.Y - EdgeUV1.Y * EdgeUV2.X);
 
-			FVector2d UVLength;
-			UVLength.X = (EdgeUV2.Y * Edge1 - EdgeUV1.Y * Edge2).Length();
-			UVLength.Y = (-EdgeUV2.X * Edge1 + EdgeUV1.X * Edge2).Length();
-
-			Chart.WorldScale += UVLength;
+			UVLengthSum += UVLength;
+			WorldLengthSum += WorldLength;
 			Chart.UVArea += UVArea;
 		}
 
-		Chart.WorldScale /= FMathd::Max(Chart.UVArea, 1e-8f);
+		if (!bScaleIslandsByWorldSpaceTexelRatio || UVLengthSum < FMathd::ZeroTolerance)
+		{
+			Chart.ScaleToWorld = 1;
+		}
+		else
+		{
+			Chart.ScaleToWorld = WorldLengthSum / UVLengthSum;
+		}
 	}
 
 
diff --git a/Engine/Plugins/Experimental/GeometryProcessing/Source/GeometricObjects/Public/Parameterization/UVPacking.h b/Engine/Plugins/Experimental/GeometryProcessing/Source/GeometricObjects/Public/Parameterization/UVPacking.h
--- a/Engine/Plugins/Experimental/GeometryProcessing/Source/GeometricObjects/Public/Parameterization/UVPacking.h
+++ b/Engine/Plugins/Experimental/GeometryProcessing/Source/GeometricObjects/Public/Parameterization/UVPacking.h
@@ -41,6 +41,8 @@ public:
 	/** If true, islands can be flipped in addition to rotate/translate/scale */
 	bool bAllowFlips = false;
 
+	/** Rescale islands to match texel-to-world-space ratio across islands */
+	bool bScaleIslandsByWorldSpaceTexelRatio = false;
 
 	/**
 	 * Standard UE4 UV layout, similar to that used for Lightmap UVs. 
diff --git a/Engine/Plugins/Experimental/PlanarCutPlugin/Source/PlanarCut/Private/FractureAutoUV.cpp b/Engine/Plugins/Experimental/PlanarCutPlugin/Source/PlanarCut/Private/FractureAutoUV.cpp
--- a/Engine/Plugins/Experimental/PlanarCutPlugin/Source/PlanarCut/Private/FractureAutoUV.cpp
+++ b/Engine/Plugins/Experimental/PlanarCutPlugin/Source/PlanarCut/Private/FractureAutoUV.cpp
@@ -6,9 +6,14 @@
 
 #include "GeometryCollection/GeometryCollectionAlgo.h"
 
+#include "DynamicMesh3.h"
+#include "Parameterization/DynamicMeshUVEditor.h"
+
 #include "Parameterization/UVPacking.h"
 
 #include "Image/ImageOccupancyMap.h"
+#include "VectorUtil.h"
+#include "FrameTypes.h"
 
 #if WITH_EDITOR
 #include "Misc/ScopedSlowTask.h"
@@ -247,7 +252,109 @@ bool UVLayout(
 	TArray<TArray<int32>> UVIslands;
 	UE::UVPacking::CreateUVIslandsFromMeshTopology<FGeomMesh>(UVMesh, UVIslands);
 
+	// detect and fix islands that don't have proper UVs (UVs all zero or otherwise collapsed to a point)
+	bool bRecreateUVsForDegenerateIslands = true;
+	if (bRecreateUVsForDegenerateIslands)
+	{
+		TArray<int> IslandVert; IslandVert.SetNumUninitialized(UVMesh.MaxVertexID());
+
+		const int32 NumIslands = UVIslands.Num();
+		for (int IslandIdx = 0; IslandIdx < NumIslands; IslandIdx++)
+		{
+			double UVArea = 0;
+			FVector3d AvgNormal(0, 0, 0);
+			FVector3d AnyPt(0, 0, 0);
+			for (int TID : UVIslands[IslandIdx])
+			{
+				FIndex3i UVInds = UVMesh.GetUVTriangle(TID);
+				FVector2f UVA = UVMesh.GetUV(UVInds.A), UVB = UVMesh.GetUV(UVInds.B), UVC = UVMesh.GetUV(UVInds.C);
+				UVArea += (double)VectorUtil::Area(UVA, UVB, UVC);
+				FVector3d VA, VB, VC;
+				UVMesh.GetTriVertices(TID, VA, VB, VC);
+				double Area;
+				AnyPt = VA;
+				FVector3d Normal = VectorUtil::NormalArea(VA, VB, VC, Area);
+				AvgNormal += Area * Normal;
+			}
+			Normalize(AvgNormal);
+			bool bDoProjection = FMathd::Abs(UVArea) < FMathd::ZeroTolerance;
+			if (bDoProjection)
+			{
+				// convert the island to a dynamic mesh so we can use the expmap UVs
+				for (int VID = 0; VID < IslandVert.Num(); ++VID)
+				{
+					IslandVert[VID] = -1;
+				}
+				TArray<int> InvIslandVert; InvIslandVert.Reserve(3 * UVIslands[IslandIdx].Num());
+				FDynamicMesh3 IslandMesh;
+				for (int TID : UVIslands[IslandIdx])
+				{
+					FIndex3i Tri = UVMesh.GetUVTriangle(TID);
+					FIndex3i NewTri = FIndex3i::Invalid();
+					for (int SubIdx = 0; SubIdx < 3; SubIdx++)
+					{
+						int SrcVID = Tri[SubIdx];
+						int& VID = IslandVert[SrcVID];
+						if (VID == -1)
+						{
+							VID = IslandMesh.AppendVertex(UVMesh.GetVertex(SrcVID));
+							check(InvIslandVert.Num() == VID);
+							InvIslandVert.Add(SrcVID);
+						}
+						NewTri[SubIdx] = VID;
+					}
+					IslandMesh.AppendTriangle(NewTri);
+				}
+				FDynamicMeshUVEditor UVEditor(&IslandMesh, 0, true);
+				TArray<int> Tris; Tris.Reserve(IslandMesh.TriangleCount());
+				for (int TID : IslandMesh.TriangleIndicesItr())
+				{
+					Tris.Add(TID);
+				}
+				bool bExpMapFailed = !UVEditor.SetTriangleUVsFromExpMap(Tris);
+				if (!bExpMapFailed)
+				{
+					// transfer UVs from the overlay
+					FDynamicMeshUVOverlay* UVOverlay = IslandMesh.Attributes()->PrimaryUV();
+					for (int IslandTID : IslandMesh.TriangleIndicesItr())
+					{
+						FIndex3i Tri = IslandMesh.GetTriangle(IslandTID);
+						FIndex3i UVTri = UVOverlay->GetTriangle(IslandTID);
+						if (!ensure(UVTri != FIndex3i::Invalid()))
+						{
+							// we shouldn't have unset tris like this if the ExpMap returned success
+							bExpMapFailed = true;
+							break;
+						}
+						for (int SubIdx = 0; SubIdx < 3; SubIdx++)
+						{
+							int VID = InvIslandVert[Tri[SubIdx]];
+							FVector2f UV = UVOverlay->GetElement(UVTri[SubIdx]);
+							UVMesh.SetUV(VID, UV);
+						}
+					}
+				}
+				if (bExpMapFailed)
+				{
+					// expmap failed; fall back to projecting UVs
+					FFrame3d Frame(AnyPt, AvgNormal);
+					for (int TID : UVIslands[IslandIdx])
+					{
+						FIndex3i Tri = UVMesh.GetUVTriangle(TID);
+						for (int SubIdx = 0; SubIdx < 3; SubIdx++)
+						{
+							int VID = Tri[SubIdx];
+							FVector2f UV = (FVector2f)Frame.ToPlaneUV(UVMesh.GetVertex(VID));
+							UVMesh.SetUV(VID, UV);
+						}
+					}
+				}
+			}
+		}
+	}
+
 	UE::Geometry::FUVPacker Packer;
+	Packer.bScaleIslandsByWorldSpaceTexelRatio = true; // let packer scale islands separately to have consistent texel-to-world ratio
 	Packer.bAllowFlips = false;
 	// StandardPack doesn't support the Packer.GutterSize feature, and always tries to leave a 1 texel gutter.
 	// To approximate a larger gutter, we tell it to consider a smaller output resolution --
