Protect against planar cut failing to triangulate the boundary of a cutting cell

#rnx

diff --git a/Engine/Plugins/Experimental/PlanarCutPlugin/Source/PlanarCut/Private/PlanarCut.cpp b/Engine/Plugins/Experimental/PlanarCutPlugin/Source/PlanarCut/Private/PlanarCut.cpp
--- a/Engine/Plugins/Experimental/PlanarCutPlugin/Source/PlanarCut/Private/PlanarCut.cpp
+++ b/Engine/Plugins/Experimental/PlanarCutPlugin/Source/PlanarCut/Private/PlanarCut.cpp
@@ -11,6 +11,7 @@
 #include "MeshAdapter.h"
 #include "FrameTypes.h"
 #include "Polygon2.h"
+#include "CompGeom/PolygonTriangulation.h"
 
 #include "GeometryCollection/GeometryCollectionAlgo.h"
 #include "GeometryCollection/GeometryCollectionClusteringUtility.h"
@@ -471,7 +472,14 @@ struct FCellMeshes
 		FMeshConstraints Constraints;
 
 		FMeshBoundaryLoops Boundary(&Mesh);
-		checkSlow(Boundary.GetLoopCount() == 1);
+		int LoopCount = Boundary.GetLoopCount();
+		if (!ensureMsgf(LoopCount == 1, TEXT("Expected to remesh a patch with a single boundary but found %d boundary loops"), LoopCount))
+		{
+			if (LoopCount == 0)
+			{
+				return;
+			}
+		}
 
 		for (int VID : Mesh.VertexIndicesItr())
 		{
@@ -910,27 +918,33 @@ private:
 				Triangulation.FillRule = FConstrainedDelaunay2f::EFillRule::NonZero;
 				Triangulation.Add(GeneralPolygon);
 				Triangulation.Triangulate();
-
-				int MID = PlaneToMaterial(PlaneIdx);
-				for (FIndex3i Triangle : Triangulation.Triangles)
+				if (Triangulation.Triangles.Num() == 0) // fall back to ear clipping if the triangulation came back empty
 				{
-					int TID = Mesh.AppendTriangle(Triangle);
-					if (ensure(TID > -1))
+					PolygonTriangulation::TriangulateSimplePolygon(Polygon.GetVertices(), Triangulation.Triangles);
+				}
+				if (ensure(Triangulation.Triangles.Num() > 0))
+				{
+					int MID = PlaneToMaterial(PlaneIdx);
+					for (FIndex3i Triangle : Triangulation.Triangles)
 					{
-						Mesh.Attributes()->GetMaterialID()->SetNewValue(TID, MID);
+						int TID = Mesh.AppendTriangle(Triangle);
+						if (ensure(TID > -1))
+						{
+							Mesh.Attributes()->GetMaterialID()->SetNewValue(TID, MID);
+						}
 					}
-				}
 
-				RemeshForNoise(Mesh, EEdgeRefineFlags::SplitsOnly, Spacing);
-				TDynamicMeshVertexAttribute<double, 3>* OriginalPosns =
-					static_cast<TDynamicMeshVertexAttribute<double, 3>*>(Mesh.Attributes()->GetAttachedAttribute(OriginalPositionAttribute));
-				for (int VID : Mesh.VertexIndicesItr())
-				{
-					OriginalPosns->SetValue(VID, Mesh.GetVertex(VID));
-				}
-				ApplyNoise(Mesh, FVector3d(Normal), Cells.InternalSurfaceMaterials.NoiseSettings.GetValue());
+					RemeshForNoise(Mesh, EEdgeRefineFlags::SplitsOnly, Spacing);
+					TDynamicMeshVertexAttribute<double, 3>* OriginalPosns =
+						static_cast<TDynamicMeshVertexAttribute<double, 3>*>(Mesh.Attributes()->GetAttachedAttribute(OriginalPositionAttribute));
+					for (int VID : Mesh.VertexIndicesItr())
+					{
+						OriginalPosns->SetValue(VID, Mesh.GetVertex(VID));
+					}
+					ApplyNoise(Mesh, FVector3d(Normal), Cells.InternalSurfaceMaterials.NoiseSettings.GetValue());
 
-				FMeshNormals::QuickComputeVertexNormals(Mesh);
+					FMeshNormals::QuickComputeVertexNormals(Mesh);
+				}
 			}, EParallelForFlags::None);
 
 		for (int CellIdx = 0; CellIdx < NumCells; CellIdx++)
