move some planarcut.cpp internals to a header file to be able to share gc/dmesh conversion code with other files in planarcut

#rnx

diff --git a/Engine/Plugins/Experimental/PlanarCutPlugin/Source/PlanarCut/Private/GeometryMeshConversion.cpp b/Engine/Plugins/Experimental/PlanarCutPlugin/Source/PlanarCut/Private/GeometryMeshConversion.cpp
new file mode 100644
--- /dev/null
+++ b/Engine/Plugins/Experimental/PlanarCutPlugin/Source/PlanarCut/Private/GeometryMeshConversion.cpp
@@ -0,0 +1,2206 @@
+// Copyright Epic Games, Inc. All Rights Reserved.
+#include "GeometryMeshConversion.h"
+
+#include "PlanarCut.h"
+
+#include "PlanarCutPlugin.h"
+
+#include "Async/ParallelFor.h"
+#include "Spatial/FastWinding.h"
+#include "Spatial/PointHashGrid3.h"
+#include "Spatial/MeshSpatialSort.h"
+#include "Util/IndexUtil.h"
+#include "Arrangement2d.h"
+#include "MeshAdapter.h"
+#include "FrameTypes.h"
+#include "Polygon2.h"
+#include "CompGeom/PolygonTriangulation.h"
+
+#include "GeometryCollection/GeometryCollectionAlgo.h"
+#include "GeometryCollection/GeometryCollectionClusteringUtility.h"
+
+#include "DisjointSet.h"
+
+#include "DynamicMesh3.h"
+#include "DynamicMeshEditor.h"
+#include "DynamicMeshAABBTree3.h"
+#include "Selections/MeshConnectedComponents.h"
+#include "MeshTransforms.h"
+#include "Operations/MeshBoolean.h"
+#include "Operations/MeshSelfUnion.h"
+#include "Operations/MergeCoincidentMeshEdges.h"
+#include "MeshBoundaryLoops.h"
+#include "QueueRemesher.h"
+#include "DynamicVertexAttribute.h"
+#include "MeshNormals.h"
+#include "MeshTangents.h"
+#include "ConstrainedDelaunay2.h"
+
+#include "Engine/EngineTypes.h"
+#include "StaticMeshOperations.h"
+#include "MeshDescriptionToDynamicMesh.h"
+#include "DynamicMeshToMeshDescription.h"
+
+#include "Algo/Rotate.h"
+
+#if WITH_EDITOR
+#include "Misc/ScopedSlowTask.h"
+#endif
+
+using namespace UE::Geometry;
+
+
+namespace UE
+{
+namespace PlanarCut
+{
+
+// functions to setup geometry collection attributes on dynamic meshes
+namespace AugmentedDynamicMesh
+{
+	FName TangentUAttribName = "TangentUAttrib";
+	FName TangentVAttribName = "TangentVAttrib";
+	FName VisibleAttribName = "VisibleAttrib";
+
+	void Augment(FDynamicMesh3& Mesh)
+	{
+		Mesh.EnableVertexColors(FVector3f(1, 1, 1));
+		Mesh.EnableVertexNormals(FVector3f::UnitZ());
+		Mesh.EnableVertexUVs(FVector2f(0, 0));
+		Mesh.EnableAttributes();
+		Mesh.Attributes()->EnableMaterialID();
+		ensure(Mesh.Attributes()->NumAttachedAttributes() == 0);
+		Mesh.Attributes()->AttachAttribute(TangentUAttribName, new TDynamicMeshVertexAttribute<float, 3>(&Mesh));
+		Mesh.Attributes()->AttachAttribute(TangentVAttribName, new TDynamicMeshVertexAttribute<float, 3>(&Mesh));
+		TDynamicMeshScalarTriangleAttribute<bool>* VisAttrib = new TDynamicMeshScalarTriangleAttribute<bool>(&Mesh);
+		VisAttrib->Initialize(true);
+		Mesh.Attributes()->AttachAttribute(VisibleAttribName, VisAttrib);
+	}
+
+	bool IsAugmented(const FDynamicMesh3& Mesh)
+	{
+		return Mesh.HasAttributes()
+			&& Mesh.Attributes()->NumAttachedAttributes() == 3
+			&& Mesh.Attributes()->HasAttachedAttribute(TangentUAttribName)
+			&& Mesh.Attributes()->HasAttachedAttribute(TangentVAttribName)
+			&& Mesh.Attributes()->HasAttachedAttribute(VisibleAttribName)
+			&& Mesh.Attributes()->HasMaterialID()
+			&& Mesh.HasVertexColors()
+			&& Mesh.HasVertexNormals()
+			&& Mesh.HasVertexUVs();
+	}
+
+	void SetDefaultAttributes(FDynamicMesh3& Mesh, bool bGlobalVisibility)
+	{
+		checkSlow(IsAugmented(Mesh));
+		TDynamicMeshVertexAttribute<float, 3>* Us =
+			static_cast<TDynamicMeshVertexAttribute<float, 3>*>(Mesh.Attributes()->GetAttachedAttribute(TangentUAttribName));
+		TDynamicMeshVertexAttribute<float, 3>* Vs =
+			static_cast<TDynamicMeshVertexAttribute<float, 3>*>(Mesh.Attributes()->GetAttachedAttribute(TangentVAttribName));
+
+		for (int VID : Mesh.VertexIndicesItr())
+		{
+			FVector3f N = Mesh.GetVertexNormal(VID);
+			FVector3f U, V;
+			VectorUtil::MakePerpVectors(N, U, V);
+			Us->SetValue(VID, U);
+			Vs->SetValue(VID, V);
+		}
+
+		TDynamicMeshScalarTriangleAttribute<bool>* Visible =
+			static_cast<TDynamicMeshScalarTriangleAttribute<bool>*>(Mesh.Attributes()->GetAttachedAttribute(VisibleAttribName));
+		for (int TID : Mesh.TriangleIndicesItr())
+		{
+			Visible->SetNewValue(TID, bGlobalVisibility);
+		}
+	}
+
+	void SetVisibility(FDynamicMesh3& Mesh, int TID, bool bIsVisible)
+	{
+		checkSlow(IsAugmented(Mesh));
+		TDynamicMeshScalarTriangleAttribute<bool>* Visible =
+			static_cast<TDynamicMeshScalarTriangleAttribute<bool>*>(Mesh.Attributes()->GetAttachedAttribute(VisibleAttribName));
+		Visible->SetValue(TID, bIsVisible);
+	}
+
+	bool GetVisibility(const FDynamicMesh3& Mesh, int TID)
+	{
+		checkSlow(IsAugmented(Mesh));
+		TDynamicMeshScalarTriangleAttribute<bool>* Visible =
+			static_cast<TDynamicMeshScalarTriangleAttribute<bool>*>(Mesh.Attributes()->GetAttachedAttribute(VisibleAttribName));
+		return Visible->GetValue(TID);
+	}
+
+	void SetTangent(FDynamicMesh3& Mesh, int VID, FVector3f Normal, FVector3f TangentU, FVector3f TangentV)
+	{
+		checkSlow(IsAugmented(Mesh));
+		TDynamicMeshVertexAttribute<float, 3>* Us =
+			static_cast<TDynamicMeshVertexAttribute<float, 3>*>(Mesh.Attributes()->GetAttachedAttribute(TangentUAttribName));
+		TDynamicMeshVertexAttribute<float, 3>* Vs =
+			static_cast<TDynamicMeshVertexAttribute<float, 3>*>(Mesh.Attributes()->GetAttachedAttribute(TangentVAttribName));
+		Us->SetValue(VID, TangentU);
+		Vs->SetValue(VID, TangentV);
+	}
+
+	void GetTangent(const FDynamicMesh3& Mesh, int VID, FVector3f& U, FVector3f& V)
+	{
+		checkSlow(IsAugmented(Mesh));
+		TDynamicMeshVertexAttribute<float, 3>* Us =
+			static_cast<TDynamicMeshVertexAttribute<float, 3>*>(Mesh.Attributes()->GetAttachedAttribute(TangentUAttribName));
+		TDynamicMeshVertexAttribute<float, 3>* Vs =
+			static_cast<TDynamicMeshVertexAttribute<float, 3>*>(Mesh.Attributes()->GetAttachedAttribute(TangentVAttribName));
+		FVector3f Normal = Mesh.GetVertexNormal(VID);
+		Us->GetValue(VID, U);
+		Vs->GetValue(VID, V);
+	}
+
+	void InitializeOverlayToPerVertexUVs(FDynamicMesh3& Mesh)
+	{
+		FDynamicMeshUVOverlay* UVs = Mesh.Attributes()->PrimaryUV();
+		UVs->ClearElements();
+		TArray<int> VertToUVMap;
+		VertToUVMap.SetNumUninitialized(Mesh.MaxVertexID());
+		for (int VID : Mesh.VertexIndicesItr())
+		{
+			FVector2f UV = Mesh.GetVertexUV(VID);
+			int UVID = UVs->AppendElement(UV);
+			VertToUVMap[VID] = UVID;
+		}
+		for (int TID : Mesh.TriangleIndicesItr())
+		{
+			FIndex3i Tri = Mesh.GetTriangle(TID);
+			Tri.A = VertToUVMap[Tri.A];
+			Tri.B = VertToUVMap[Tri.B];
+			Tri.C = VertToUVMap[Tri.C];
+			UVs->SetTriangle(TID, Tri);
+		}
+	}
+
+	void InitializeOverlayToPerVertexTangents(FDynamicMesh3& Mesh)
+	{
+		Mesh.Attributes()->EnableTangents();
+		FDynamicMeshNormalOverlay* TangentOverlays[2] = { Mesh.Attributes()->PrimaryTangents(), Mesh.Attributes()->PrimaryBiTangents() };
+		TangentOverlays[0]->ClearElements();
+		TangentOverlays[1]->ClearElements();
+		TArray<int> VertToTangentMap;
+		VertToTangentMap.SetNumUninitialized(Mesh.MaxVertexID());
+		for (int VID : Mesh.VertexIndicesItr())
+		{
+			FVector3f Tangents[2];
+			GetTangent(Mesh, VID, Tangents[0], Tangents[1]);
+			int TID = TangentOverlays[0]->AppendElement(&Tangents[0].X);
+			int TID2 = TangentOverlays[1]->AppendElement(&Tangents[1].X);
+			check(TID == TID2);
+			VertToTangentMap[VID] = TID;
+		}
+		for (int TID : Mesh.TriangleIndicesItr())
+		{
+			FIndex3i Tri = Mesh.GetTriangle(TID);
+			Tri.A = VertToTangentMap[Tri.A];
+			Tri.B = VertToTangentMap[Tri.B];
+			Tri.C = VertToTangentMap[Tri.C];
+			TangentOverlays[0]->SetTriangle(TID, Tri);
+			TangentOverlays[1]->SetTriangle(TID, Tri);
+		}
+	}
+
+	void ComputeTangents(FDynamicMesh3& Mesh, bool bOnlyOddMaterials, const TArrayView<const int32>& WhichMaterials, bool bRecomputeNormals)
+	{
+		FDynamicMeshNormalOverlay* Normals = Mesh.Attributes()->PrimaryNormals();
+		FMeshNormals::InitializeOverlayToPerVertexNormals(Normals, !bRecomputeNormals);
+		if (bRecomputeNormals)
+		{
+			FMeshNormals::QuickRecomputeOverlayNormals(Mesh);
+		}
+
+		// Copy per-vertex UVs to a UV overlay, because that's what the tangents code uses
+		// (TODO: consider making a tangent computation path that uses vertex normals / UVs)
+		InitializeOverlayToPerVertexUVs(Mesh);
+		FDynamicMeshUVOverlay* UVs = Mesh.Attributes()->PrimaryUV();
+
+		FComputeTangentsOptions Options;
+		Options.bAngleWeighted = true;
+		Options.bAveraged = true;
+		FMeshTangentsf Tangents(&Mesh);
+		Tangents.ComputeTriVertexTangents(Normals, UVs, Options);
+
+		const TArray<FVector3<float>>& TanU = Tangents.GetTangents();
+		const TArray<FVector3<float>>& TanV = Tangents.GetBitangents();
+		FDynamicMeshMaterialAttribute* MaterialIDs = Mesh.Attributes()->GetMaterialID();
+		for (int TID : Mesh.TriangleIndicesItr())
+		{
+			int MaterialID = MaterialIDs->GetValue(TID);
+			if (bOnlyOddMaterials && (MaterialID % 2) == 0)
+			{
+				continue;
+			}
+			else if (WhichMaterials.Contains(MaterialID))
+			{
+				continue;
+			}
+					
+			int TanIdxBase = TID * 3;
+			FIndex3i Tri = Mesh.GetTriangle(TID);
+			for (int Idx = 0; Idx < 3; Idx++)
+			{
+				int VID = Tri[Idx];
+				int TanIdx = TanIdxBase + Idx;
+				SetTangent(Mesh, VID, Mesh.GetVertexNormal(VID), TanU[TanIdx], TanV[TanIdx]);
+			}
+		}
+	}
+
+	// per component sampling is a rough heuristic to avoid doing geodesic distance but still get points on a 'thin' slice
+	void AddCollisionSamplesPerComponent(FDynamicMesh3& Mesh, double Spacing)
+	{
+		checkSlow(IsAugmented(Mesh));
+		FMeshConnectedComponents Components(&Mesh);
+		// TODO: if/when we switch to merged edges representation, pass a predicate here based on whether there's a normal seam ?
+		Components.FindConnectedTriangles();
+		TArray<TPointHashGrid3d<int>> KnownSamples;  KnownSamples.Reserve(Components.Num());
+		for (int ComponentIdx = 0; ComponentIdx < Components.Num(); ComponentIdx++)
+		{
+			KnownSamples.Emplace(.5 * Spacing / FMathd::InvSqrt3, -1);
+		}
+
+		TArray<int> AlreadySeen; AlreadySeen.Init(-1, Mesh.MaxVertexID());
+		for (int ComponentIdx = 0; ComponentIdx < Components.Num(); ComponentIdx++)
+		{
+			FMeshConnectedComponents::FComponent& Component = Components.GetComponent(ComponentIdx);
+			for (int TID : Component.Indices)
+			{
+				FIndex3i Tri = Mesh.GetTriangle(TID);
+				for (int SubIdx = 0; SubIdx < 3; SubIdx++)
+				{
+					int VID = Tri[SubIdx];
+					if (AlreadySeen[VID] != ComponentIdx)
+					{
+						AlreadySeen[VID] = ComponentIdx;
+						KnownSamples[ComponentIdx].InsertPointUnsafe(VID, Mesh.GetVertex(VID));
+					}
+				}
+			}
+		}
+		AlreadySeen.Empty();
+
+		double SpacingThreshSq = Spacing * Spacing; // if points are more than Spacing apart, consider adding a new point between them
+		for (int ComponentIdx = 0; ComponentIdx < Components.Num(); ComponentIdx++)
+		{
+			FMeshConnectedComponents::FComponent& Component = Components.GetComponent(ComponentIdx);
+			for (int TID : Component.Indices)
+			{
+				FIndex3i TriVIDs = Mesh.GetTriangle(TID);
+				FTriangle3d Triangle;
+				Mesh.GetTriVertices(TID, Triangle.V[0], Triangle.V[1], Triangle.V[2]);
+				double EdgeLensSq[3];
+				int MaxEdgeIdx = 0;
+				double MaxEdgeLenSq = 0;
+				for (int i = 2, j = 0; j < 3; i = j++)
+				{
+					double EdgeLenSq = DistanceSquared(Triangle.V[i], Triangle.V[j]);
+					if (EdgeLenSq > MaxEdgeLenSq)
+					{
+						MaxEdgeIdx = i;
+						MaxEdgeLenSq = EdgeLenSq;
+					}
+					EdgeLensSq[i] = EdgeLenSq;
+				}
+				// if we found a too-long edge, we can try sampling the tri
+				if (MaxEdgeLenSq > SpacingThreshSq)
+				{
+					FVector3f Normal = (FVector3f)VectorUtil::Normal(Triangle.V[0], Triangle.V[1], Triangle.V[2]);
+
+					// Pick number of samples based on the longest edge
+					double LongEdgeLen = FMathd::Sqrt(MaxEdgeLenSq);
+					int Divisions = FMathd::Floor(LongEdgeLen / Spacing);
+					double Factor = 1.0 / double(Divisions + 1);
+					int SecondEdgeIdx = (MaxEdgeIdx + 1) % 3;
+					int ThirdEdgeIdx = (MaxEdgeIdx + 2) % 3;
+					// Sample along the two longest edges first, then interpolate these samples
+					int SecondLongestEdgeIdx = SecondEdgeIdx;
+					if (EdgeLensSq[SecondEdgeIdx] < EdgeLensSq[ThirdEdgeIdx])
+					{
+						SecondLongestEdgeIdx = ThirdEdgeIdx;
+					}
+					int SecondLongestSecondEdgeIdx = (SecondLongestEdgeIdx + 1) % 3;
+					for (int DivI = 0; DivI < Divisions; DivI++)
+					{
+						double Along = (DivI + 1) * Factor;
+						FVector3d E1Bary(0, 0, 0), E2Bary(0, 0, 0);
+						E1Bary[MaxEdgeIdx] = Along;
+						E1Bary[SecondEdgeIdx] = 1 - Along;
+						E2Bary[SecondLongestEdgeIdx] = 1 - Along;
+						E2Bary[SecondLongestSecondEdgeIdx] = Along;
+
+						// Choose number of samples between the two edge points based on their distance
+						double AcrossDist = Triangle.BarycentricPoint(E1Bary).Distance(Triangle.BarycentricPoint(E2Bary));
+						int DivisionsAcross = FMathd::Ceil(AcrossDist / Spacing);
+						double FactorAcross = 1.0 / double(DivisionsAcross + 1);
+						for (int DivJ = 0; DivJ < DivisionsAcross; DivJ++)
+						{
+							double AlongAcross = (DivJ + 1) * FactorAcross;
+							FVector3d Bary = UE::Geometry::Lerp(E1Bary, E2Bary, AlongAcross);
+							FVector3d SamplePos = Triangle.BarycentricPoint(Bary);
+							if (!KnownSamples[ComponentIdx].IsCellEmptyUnsafe(SamplePos)) // fast early out; def. have pt within radius
+							{
+								continue;
+							}
+							TPair<int, double> VIDDist = KnownSamples[ComponentIdx].FindNearestInRadius(SamplePos, Spacing * .5, [&Mesh, SamplePos](int VID)
+								{
+									return DistanceSquared(Mesh.GetVertex(VID), SamplePos);
+								});
+							// No point within radius Spacing/2 -> Add a new sample
+							if (VIDDist.Key == -1)
+							{
+								// no point within radius; can add a sample here
+								FVertexInfo Info(SamplePos, Normal);
+
+								int AddedVID = Mesh.AppendVertex(Info);
+								KnownSamples[ComponentIdx].InsertPointUnsafe(AddedVID, SamplePos);
+							}
+						}
+					}
+				}
+			}
+		}
+	}
+}
+
+
+void SetGeometryCollectionAttributes(FDynamicMesh3& Mesh)
+{
+	AugmentedDynamicMesh::Augment(Mesh);
+}
+
+
+FCellMeshes::FCellMeshes(const FPlanarCells& Cells, FAxisAlignedBox3d DomainBounds, double Grout, double ExtendDomain, bool bIncludeOutsideCell)
+{
+	Init(Cells, DomainBounds, Grout, ExtendDomain, bIncludeOutsideCell);
+}
+
+
+FCellMeshes::FCellMeshes(FDynamicMesh3& SingleCutter, const FInternalSurfaceMaterials& Materials, TOptional<FTransform> Transform)
+{
+	SetNumCells(2);
+
+	if (Transform.IsSet())
+	{
+		MeshTransforms::ApplyTransform(SingleCutter, FTransform3d(Transform.GetValue()));
+	}
+
+	// Mesh should already be augmented
+	if (!ensure(AugmentedDynamicMesh::IsAugmented(SingleCutter)))
+	{
+		AugmentedDynamicMesh::Augment(SingleCutter);
+	}
+
+	CellMeshes[0].AugMesh = SingleCutter;
+
+	// first mesh is the same as the second mesh, but will be subtracted b/c it's the "outside cell"
+	// TODO: special case this logic so we don't have to hold two copies of the exact same mesh!
+	CellMeshes[1].AugMesh = CellMeshes[0].AugMesh;
+	OutsideCellIndex = 1;
+
+}
+
+
+// Special function to just make the "grout" part of the planar mesh cells
+// Used to make the multi-plane cuts with grout easier to implement
+void FCellMeshes::MakeOnlyPlanarGroutCell(const FPlanarCells& Cells, FAxisAlignedBox3d DomainBounds, double Grout)
+{
+	CellMeshes.Reset();
+
+	if (!ensure(Grout > 0) || !ensure(Cells.IsInfinitePlane()))
+	{
+		return;
+	}
+
+	float GlobalUVScale = Cells.InternalSurfaceMaterials.GlobalUVScale;
+	if (!ensure(GlobalUVScale > 0))
+	{
+		GlobalUVScale = 1;
+	}
+
+	SetNumCells(1);
+
+	bool bNoise = Cells.InternalSurfaceMaterials.NoiseSettings.IsSet();
+
+	double ExtendDomain = bNoise ? Cells.InternalSurfaceMaterials.NoiseSettings->Amplitude : 0;
+	DomainBounds.Expand(ExtendDomain);
+
+	CreateMeshesForSinglePlane(Cells, DomainBounds, bNoise, GlobalUVScale, Grout, true);
+
+	for (FCellInfo& CellInfo : CellMeshes)
+	{
+		AugmentedDynamicMesh::SetDefaultAttributes(CellInfo.AugMesh, Cells.InternalSurfaceMaterials.bGlobalVisibility);
+	}
+}
+
+
+void FCellMeshes::RemeshForNoise(FDynamicMesh3& Mesh, EEdgeRefineFlags EdgeFlags, double TargetEdgeLen)
+{
+	FQueueRemesher Remesh(&Mesh);
+	Remesh.bPreventNormalFlips = true;
+	FMeshConstraints Constraints;
+
+	FMeshBoundaryLoops Boundary(&Mesh);
+	int LoopCount = Boundary.GetLoopCount();
+	if (!ensureMsgf(LoopCount == 1, TEXT("Expected to remesh a patch with a single boundary but found %d boundary loops"), LoopCount))
+	{
+		if (LoopCount == 0)
+		{
+			return;
+		}
+	}
+
+	for (int VID : Mesh.VertexIndicesItr())
+	{
+		FVertexConstraint FullyConstrain(true, false, VID);
+		Constraints.SetOrUpdateVertexConstraint(VID, FullyConstrain);
+	}
+
+
+	FEdgeConstraint EdgeConstraint(EdgeFlags);
+	for (int EID : Boundary[0].Edges)
+	{
+		Constraints.SetOrUpdateEdgeConstraint(EID, EdgeConstraint);
+	}
+	Remesh.SetExternalConstraints(Constraints);
+	Remesh.SetTargetEdgeLength(TargetEdgeLen);
+	Remesh.Precompute();
+	Remesh.FastestRemesh();
+}
+
+
+float FCellMeshes::OctaveNoise(const FVector& V, const FNoiseSettings& Settings)
+{
+	int32 Octaves = Settings.Octaves;
+	float NoiseValue = 0;
+	float OctaveScale = 1;
+	for (int32 Octave = 0; Octave < Octaves; Octave++, OctaveScale *= 2)
+	{
+		NoiseValue += FMath::PerlinNoise3D(V * OctaveScale) / OctaveScale;
+	}
+	return NoiseValue;
+}
+
+
+FVector FCellMeshes::NoiseVector(const FVector& Pos, const FNoiseSettings& Settings)
+{
+	float Frequency = Settings.Frequency;
+	FVector Base = Pos * Frequency;
+	return FVector(
+		OctaveNoise(Base + NoiseOffsetX, Settings),
+		OctaveNoise(Base + NoiseOffsetY, Settings),
+		OctaveNoise(Base + NoiseOffsetZ, Settings)
+	) * Settings.Amplitude;
+}
+
+
+FVector3d FCellMeshes::NoiseDisplacement(const FVector3d& Pos, const FNoiseSettings& Settings)
+{
+	FVector P = FVector(Pos);
+	return FVector3d(NoiseVector(P, Settings));
+}
+
+
+void FCellMeshes::ApplyNoise(FDynamicMesh3& Mesh, FVector3d Normal, const FNoiseSettings& Settings, bool bProjectBoundariesToNormal)
+{
+	double Amplitude = (double)Settings.Amplitude;
+	double Frequency = (double)Settings.Frequency;
+	int32 Octaves = Settings.Octaves;
+	FVector3d Z = Normal * Amplitude;
+
+	for (int VID : Mesh.VertexIndicesItr())
+	{
+		FVector3d Pos = Mesh.GetVertex(VID);
+		FVector3d Displacement = NoiseDisplacement(Pos, Settings);
+		if (bProjectBoundariesToNormal || !Mesh.IsBoundaryVertex(VID))
+		{
+			// project displacement onto the normal direction
+			Displacement = Normal * Displacement.Dot(Normal);
+		}
+
+
+		Mesh.SetVertex(VID, Pos + Displacement);
+	}
+}
+
+
+void FCellMeshes::Init(const FPlanarCells& Cells, FAxisAlignedBox3d DomainBounds, double Grout, double ExtendDomain, bool bIncludeOutsideCell)
+{
+	InitEmpty();
+
+	float GlobalUVScale = Cells.InternalSurfaceMaterials.GlobalUVScale;
+	if (!ensure(GlobalUVScale > 0))
+	{
+		GlobalUVScale = 1;
+	}
+
+	int NumCells = Cells.NumCells;
+	bool bHasGroutCell = Grout > 0;
+	if (bIncludeOutsideCell && !Cells.IsInfinitePlane())
+	{
+		OutsideCellIndex = NumCells;
+		NumCells++;
+	}
+
+	SetNumCells(NumCells);
+
+	bool bNoise = Cells.InternalSurfaceMaterials.NoiseSettings.IsSet();
+	if (bNoise)
+	{
+		ExtendDomain += Cells.InternalSurfaceMaterials.NoiseSettings->Amplitude;
+	}
+	DomainBounds.Expand(ExtendDomain);
+
+	// special handling for the infinite plane case; we need to adapt this to be a closed volume
+	if (Cells.IsInfinitePlane())
+	{
+		CreateMeshesForSinglePlane(Cells, DomainBounds, bNoise, GlobalUVScale, Grout, false);
+	}
+	else
+	{
+		if (!bNoise) // bounded cells w/ no noise
+		{
+			CreateMeshesForBoundedPlanesWithoutNoise(NumCells, Cells, DomainBounds, bNoise, GlobalUVScale);
+		}
+		else // bounded cells with noise -- make each boundary plane separately so we can remesh them w/ noise vertices
+		{
+			CreateMeshesForBoundedPlanesWithNoise(NumCells, Cells, DomainBounds, bNoise, GlobalUVScale);
+		}
+		ApplyGeneralGrout(Grout);
+	}
+
+	// TODO: self-union on cells when it makes sense to do so (for non-single-plane inputs w/ high noise or possible untracked adjacencies)
+	/*for (FCellInfo& CellInfo : CellMeshes)
+	{
+		FMeshSelfUnion SelfUnion(&CellInfo.AugMesh);
+		// TODO: need to have an option in SelfUnion to not weld edges
+		SelfUnion.Compute();
+	}*/
+
+	for (FCellInfo& CellInfo : CellMeshes)
+	{
+		AugmentedDynamicMesh::SetDefaultAttributes(CellInfo.AugMesh, Cells.InternalSurfaceMaterials.bGlobalVisibility);
+	}
+}
+
+
+void FCellMeshes::ApplyGeneralGrout(double Grout)
+{
+	if (Grout <= 0)
+	{
+		return;
+	}
+
+	// apply grout to all cells
+	for (int MeshIdx = 0; MeshIdx < CellMeshes.Num(); MeshIdx++)
+	{
+		if (MeshIdx == OutsideCellIndex)
+		{
+			continue;
+		}
+
+		FDynamicMesh3& Mesh = CellMeshes[MeshIdx].AugMesh;
+		// TODO: scale from mesh center of mass instead of the vertex centroid?
+		FVector3d VertexCentroid(0, 0, 0);
+		for (FVector3d V : Mesh.VerticesItr())
+		{
+			VertexCentroid += V;
+		}
+		VertexCentroid /= (double)Mesh.VertexCount();
+		FAxisAlignedBox3d Bounds = Mesh.GetCachedBounds();
+		double BoundsSize = Bounds.MaxDim();
+		// currently just scale the meshes down so they leave half-a-grout worth of space on their longest axis
+		// or delete the mesh if it's so small that that would require a negative scale
+		// TODO: consider instead computing a true offset mesh
+		//  (note that we don't currently have a good UV-preserving+sharp-edge-preserving way to do that)
+		double ScaleFactor = (BoundsSize - Grout * .5) / BoundsSize;
+		if (ScaleFactor < FMathd::ZeroTolerance * 1000)
+		{
+			// if the grout scale factor would be ~zero or negative, just clear the mesh instead
+			Mesh.Clear();
+			AugmentedDynamicMesh::Augment(Mesh);
+		}
+		else
+		{
+			MeshTransforms::Scale(Mesh, FVector3d::One() * ScaleFactor, VertexCentroid);
+		}
+	}
+
+	// create outside cell (if there is room for it) by appending all the other meshes
+	if (OutsideCellIndex != -1)
+	{
+		FDynamicMesh3& OutsideMesh = CellMeshes[OutsideCellIndex].AugMesh;
+		OutsideMesh.Clear();
+		AugmentedDynamicMesh::Augment(OutsideMesh);
+		FDynamicMeshEditor OutsideMeshEditor(&OutsideMesh);
+		for (int MeshIdx = 0; MeshIdx < CellMeshes.Num(); MeshIdx++)
+		{
+			if (MeshIdx == OutsideCellIndex)
+			{
+				continue;
+			}
+			FMeshIndexMappings IndexMaps;
+			OutsideMeshEditor.AppendMesh(&CellMeshes[MeshIdx].AugMesh, IndexMaps);
+		}
+	}
+}
+
+
+void FCellMeshes::AppendMesh(FDynamicMesh3& Base, FDynamicMesh3& ToAppend, bool bFlipped)
+{
+	FDynamicMeshEditor Editor(&Base);
+	FMeshIndexMappings Mapping;
+	Editor.AppendMesh(&ToAppend, Mapping);
+	if (bFlipped)
+	{
+		for (int TID : ToAppend.TriangleIndicesItr())
+		{
+			Base.ReverseTriOrientation(Mapping.GetNewTriangle(TID));
+		}
+		for (int VID : ToAppend.VertexIndicesItr())
+		{
+			int BaseVID = Mapping.GetNewVertex(VID);
+			Base.SetVertexNormal(BaseVID, -Base.GetVertexNormal(BaseVID));
+		}
+	}
+}
+
+
+
+
+void FCellMeshes::CreateMeshesForBoundedPlanesWithoutNoise(int NumCells, const FPlanarCells& Cells, const FAxisAlignedBox3d& DomainBounds, bool bNoise, double GlobalUVScale)
+{
+	for (int32 PlaneIdx = 0; PlaneIdx < Cells.PlaneCells.Num(); PlaneIdx++)
+	{
+		const TPair<int32, int32>& CellPair = Cells.PlaneCells[PlaneIdx];
+		FDynamicMesh3* Meshes[2]{ &CellMeshes[CellPair.Key].AugMesh, nullptr };
+		int32 OtherCell = CellPair.Value < 0 ? OutsideCellIndex : CellPair.Value;
+		int NumMeshes = OtherCell < 0 ? 1 : 2;
+		if (NumMeshes == 2)
+		{
+			Meshes[1] = &CellMeshes[OtherCell].AugMesh;
+		}
+
+		const TArray<int>& PlaneBoundary = Cells.PlaneBoundaries[PlaneIdx];
+		FVector3f Normal(Cells.Planes[PlaneIdx].GetNormal());
+		FFrame3d PlaneFrame(Cells.Planes[PlaneIdx]);
+		FVertexInfo PlaneVertInfo;
+		PlaneVertInfo.bHaveC = true;
+		PlaneVertInfo.bHaveUV = true;
+		PlaneVertInfo.bHaveN = true;
+		PlaneVertInfo.Color = FVector3f(1, 1, 1);
+		int VertStart[2]{ -1, -1 };
+		for (int MeshIdx = 0; MeshIdx < NumMeshes; MeshIdx++)
+		{
+			PlaneVertInfo.Normal = Normal;
+			if (MeshIdx == 1 && OtherCell != OutsideCellIndex)
+			{
+				PlaneVertInfo.Normal *= -1.0f;
+			}
+			VertStart[MeshIdx] = Meshes[MeshIdx]->MaxVertexID();
+			FVector2f MinUV(FMathf::MaxReal, FMathf::MaxReal);
+			for (int BoundaryVertex : PlaneBoundary)
+			{
+				FVector3d Position = FVector3d(Cells.PlaneBoundaryVertices[BoundaryVertex]);
+				FVector2f UV = FVector2f(PlaneFrame.ToPlaneUV(Position));
+				MinUV.X = FMathf::Min(UV.X, MinUV.X);
+				MinUV.Y = FMathf::Min(UV.Y, MinUV.Y);
+			}
+			for (int BoundaryVertex : PlaneBoundary)
+			{
+				PlaneVertInfo.Position = FVector3d(Cells.PlaneBoundaryVertices[BoundaryVertex]);
+				PlaneVertInfo.UV = (FVector2f(PlaneFrame.ToPlaneUV(PlaneVertInfo.Position)) - MinUV) * GlobalUVScale;
+				Meshes[MeshIdx]->AppendVertex(PlaneVertInfo);
+			}
+		}
+
+		int MID = PlaneToMaterial(PlaneIdx);
+		if (Cells.AssumeConvexCells)
+		{
+			// put a fan
+			for (int V0 = 0, V1 = 1, V2 = 2; V2 < PlaneBoundary.Num(); V1 = V2++)
+			{
+				for (int MeshIdx = 0; MeshIdx < NumMeshes; MeshIdx++)
+				{
+					int Offset = VertStart[MeshIdx];
+					FIndex3i Tri(V0 + Offset, V1 + Offset, V2 + Offset);
+					if (MeshIdx == 1 && OtherCell != OutsideCellIndex)
+					{
+						Swap(Tri.B, Tri.C);
+					}
+					int TID = Meshes[MeshIdx]->AppendTriangle(Tri);
+					if (ensure(TID > -1))
+					{
+						Meshes[MeshIdx]->Attributes()->GetMaterialID()->SetNewValue(TID, MID);
+					}
+				}
+			}
+		}
+		else // cells may not be convex; cannot triangulate w/ fan
+		{
+			// Delaunay triangulate
+			FPolygon2f Polygon;
+			for (int V = 0; V < PlaneBoundary.Num(); V++)
+			{
+				Polygon.AppendVertex(Meshes[0]->GetVertexUV(VertStart[0] + V));
+			}
+
+			FGeneralPolygon2f GeneralPolygon(Polygon);
+			FConstrainedDelaunay2f Triangulation;
+			Triangulation.FillRule = FConstrainedDelaunay2f::EFillRule::NonZero;
+			Triangulation.Add(GeneralPolygon);
+			Triangulation.Triangulate();
+
+			for (int MeshIdx = 0; MeshIdx < NumMeshes; MeshIdx++)
+			{
+				int Offset = VertStart[MeshIdx];
+				for (FIndex3i Triangle : Triangulation.Triangles)
+				{
+					Triangle.A += Offset;
+					Triangle.B += Offset;
+					Triangle.C += Offset;
+					if (MeshIdx == 1 && OtherCell != OutsideCellIndex)
+					{
+						Swap(Triangle.B, Triangle.C);
+					}
+					int TID = Meshes[MeshIdx]->AppendTriangle(Triangle);
+					if (ensure(TID > -1))
+					{
+						Meshes[MeshIdx]->Attributes()->GetMaterialID()->SetNewValue(TID, MID);
+					}
+				}
+			}
+		}
+	}
+}
+
+
+double FCellMeshes::GetSafeNoiseSpacing(float SurfaceArea, float TargetSpacing)
+{
+	double MaxVerts = 1000000; // try to avoid creating more than ~ a million new vertices
+	double MinEdgeLen = FMathd::Sqrt((double)SurfaceArea / MaxVerts);
+	double Spacing = FMath::Max3(.001, MinEdgeLen, (double)TargetSpacing);
+	if (Spacing > TargetSpacing)
+	{
+		UE_LOG(LogPlanarCut, Warning,
+			TEXT("Requested spacing of noise points (surface resolution) of %f would require too many added vertices; Using %f instead."),
+			TargetSpacing, Spacing);
+	}
+	return Spacing;
+}
+
+void FCellMeshes::CreateMeshesForBoundedPlanesWithNoise(int NumCells, const FPlanarCells& Cells, const FAxisAlignedBox3d& DomainBounds, bool bNoise, double GlobalUVScale)
+{
+	TArray<FDynamicMesh3> PlaneMeshes;
+	PlaneMeshes.SetNum(Cells.Planes.Num());
+	FName OriginalPositionAttribute = "OriginalPosition";
+	for (FDynamicMesh3& PlaneMesh : PlaneMeshes)
+	{
+		PlaneMesh.EnableVertexUVs(FVector2f(0, 0));
+		PlaneMesh.EnableVertexNormals(FVector3f::UnitZ());
+		PlaneMesh.EnableVertexColors(FVector3f(1, 1, 1));
+		PlaneMesh.EnableAttributes();
+		PlaneMesh.Attributes()->EnableMaterialID();
+		PlaneMesh.Attributes()->AttachAttribute(OriginalPositionAttribute, new TDynamicMeshVertexAttribute<double, 3>(&PlaneMesh));
+	}
+
+	struct FPlaneIdxAndFlip
+	{
+		int32 PlaneIdx;
+		bool bIsFlipped;
+	};
+	TArray<TArray<FPlaneIdxAndFlip>> CellPlanes; // per cell, the planes that border that cell
+	CellPlanes.SetNum(NumCells);
+
+	for (int32 PlaneIdx = 0; PlaneIdx < Cells.PlaneCells.Num(); PlaneIdx++)
+	{
+		const TPair<int32, int32>& CellPair = Cells.PlaneCells[PlaneIdx];
+		int32 OtherCell = CellPair.Value < 0 ? OutsideCellIndex : CellPair.Value;
+		if (ensure(CellPlanes.IsValidIndex(CellPair.Key)))
+		{
+			CellPlanes[CellPair.Key].Add({ PlaneIdx, false });
+		}
+		if (CellPlanes.IsValidIndex(OtherCell))
+		{
+			CellPlanes[OtherCell].Add({ PlaneIdx, true });
+		}
+	}
+
+	// heuristic to protect against creating too many vertices on remeshing
+	float TotalArea = 0;
+	for (int32 PlaneIdx = 0; PlaneIdx < Cells.Planes.Num(); PlaneIdx++)
+	{
+		const TArray<int>& PlaneBoundary = Cells.PlaneBoundaries[PlaneIdx];
+		const FVector& V0 = Cells.PlaneBoundaryVertices[PlaneBoundary[0]];
+		FVector AreaVec = FVector::ZeroVector;
+		for (int32 V1Idx = 1, V2Idx = 2; V2Idx < PlaneBoundary.Num(); V1Idx = V2Idx++)
+		{
+			const FVector& V1 = Cells.PlaneBoundaryVertices[PlaneBoundary[V1Idx]];
+			const FVector& V2 = Cells.PlaneBoundaryVertices[PlaneBoundary[V2Idx]];
+			AreaVec += (V1 - V0) ^ (V2 - V1);
+		}
+		TotalArea += AreaVec.Size();
+	}
+	double Spacing = GetSafeNoiseSpacing(TotalArea, Cells.InternalSurfaceMaterials.NoiseSettings->PointSpacing);
+
+	ParallelFor(Cells.Planes.Num(), [this, OriginalPositionAttribute, &PlaneMeshes, &Cells, GlobalUVScale, Spacing](int32 PlaneIdx)
+		{
+			FDynamicMesh3& Mesh = PlaneMeshes[PlaneIdx];
+			const TArray<int>& PlaneBoundary = Cells.PlaneBoundaries[PlaneIdx];
+			FVector3f Normal(Cells.Planes[PlaneIdx].GetNormal());
+			FFrame3d PlaneFrame(Cells.Planes[PlaneIdx]);
+			FVertexInfo PlaneVertInfo;
+			PlaneVertInfo.bHaveC = true;
+			PlaneVertInfo.bHaveUV = true;
+			PlaneVertInfo.bHaveN = true;
+			PlaneVertInfo.Normal = Normal;
+			PlaneVertInfo.UV = FVector2f(0, 0); // UVs will be set below, after noise is added
+			PlaneVertInfo.Color = FVector3f(1, 1, 1);
+
+			FPolygon2f Polygon;
+			for (int BoundaryVertex : PlaneBoundary)
+			{
+				PlaneVertInfo.Position = FVector3d(Cells.PlaneBoundaryVertices[BoundaryVertex]);
+				Polygon.AppendVertex((FVector2f)PlaneFrame.ToPlaneUV(PlaneVertInfo.Position));
+				Mesh.AppendVertex(PlaneVertInfo);
+			}
+
+			// we do a CDT here to give a slightly better start to remeshing; we could try simple ear clipping instead
+			FGeneralPolygon2f GeneralPolygon(Polygon);
+			FConstrainedDelaunay2f Triangulation;
+			Triangulation.FillRule = FConstrainedDelaunay2f::EFillRule::NonZero;
+			Triangulation.Add(GeneralPolygon);
+			Triangulation.Triangulate();
+			if (Triangulation.Triangles.Num() == 0) // fall back to ear clipping if the triangulation came back empty
+			{
+				PolygonTriangulation::TriangulateSimplePolygon(Polygon.GetVertices(), Triangulation.Triangles);
+			}
+			if (ensure(Triangulation.Triangles.Num() > 0))
+			{
+				int MID = PlaneToMaterial(PlaneIdx);
+				for (FIndex3i Triangle : Triangulation.Triangles)
+				{
+					int TID = Mesh.AppendTriangle(Triangle);
+					if (ensure(TID > -1))
+					{
+						Mesh.Attributes()->GetMaterialID()->SetNewValue(TID, MID);
+					}
+				}
+
+				RemeshForNoise(Mesh, EEdgeRefineFlags::SplitsOnly, Spacing);
+				TDynamicMeshVertexAttribute<double, 3>* OriginalPosns =
+					static_cast<TDynamicMeshVertexAttribute<double, 3>*>(Mesh.Attributes()->GetAttachedAttribute(OriginalPositionAttribute));
+				for (int VID : Mesh.VertexIndicesItr())
+				{
+					OriginalPosns->SetValue(VID, Mesh.GetVertex(VID));
+				}
+				ApplyNoise(Mesh, FVector3d(Normal), Cells.InternalSurfaceMaterials.NoiseSettings.GetValue());
+
+				FMeshNormals::QuickComputeVertexNormals(Mesh);
+			}
+		}, EParallelForFlags::None);
+
+	for (int CellIdx = 0; CellIdx < NumCells; CellIdx++)
+	{
+		FCellInfo& CellInfo = CellMeshes[CellIdx];
+		FDynamicMesh3& Mesh = CellInfo.AugMesh;
+		Mesh.Attributes()->AttachAttribute(OriginalPositionAttribute, new TDynamicMeshVertexAttribute<double, 3>(&Mesh));
+		bool bFlipForOutsideCell = CellIdx == OutsideCellIndex; // outside cell will be subtracted, and needs all planes flipped vs normal
+		for (FPlaneIdxAndFlip PlaneInfo : CellPlanes[CellIdx])
+		{
+			AppendMesh(Mesh, PlaneMeshes[PlaneInfo.PlaneIdx], PlaneInfo.bIsFlipped ^ bFlipForOutsideCell);
+		}
+	}
+
+	// resolve self-intersections
+
+	// build hash grid of mesh vertices so we correspond all same-pos vertices across touching meshes
+	TPointHashGrid3d<FIndex2i> MeshesVertices(FMathd::ZeroTolerance * 1000, FIndex2i::Invalid());
+	for (int CellIdx = 0; CellIdx < NumCells; CellIdx++)
+	{
+		FCellInfo& CellInfo = CellMeshes[CellIdx];
+		FDynamicMesh3& Mesh = CellInfo.AugMesh;
+		for (int VID : Mesh.VertexIndicesItr())
+		{
+			MeshesVertices.InsertPointUnsafe(FIndex2i(CellIdx, VID), Mesh.GetVertex(VID));
+		}
+	}
+
+	// repeatedly detect and resolve collisions until there are no more (or give up after too many iterations)
+	TArray<bool> CellUnmoved; CellUnmoved.Init(false, NumCells);
+	const int MaxIters = 10;
+	for (int Iters = 0; Iters < MaxIters; Iters++)
+	{
+		struct FUpdate
+		{
+			FIndex2i Tris;
+			TArray<FIndex2i> IDs;
+			FUpdate(int TriA = -1, int TriB = -1) : Tris(TriA, TriB)
+			{}
+		};
+
+		// todo: can parallelize?
+		TArray<TArray<FUpdate>> Updates; Updates.SetNum(NumCells);
+		bool bAnyUpdatesNeeded = false;
+		for (int CellIdx = 0; CellIdx < NumCells; CellIdx++)
+		{
+			if (CellUnmoved[CellIdx])
+			{
+				// if nothing moved since last time we resolved self intersections on this cell, don't need to process again
+				continue;
+			}
+			FDynamicMesh3& Mesh = CellMeshes[CellIdx].AugMesh;
+			FDynamicMeshAABBTree3 CellTree(&Mesh, true);
+			MeshIntersection::FIntersectionsQueryResult Intersections = CellTree.FindAllSelfIntersections(true);
+			for (MeshIntersection::FSegmentIntersection& Seg : Intersections.Segments)
+			{
+				// manually check for shared edges by vertex position because they might not be topologically connected
+				FIndex3i Tri[2]{ Mesh.GetTriangle(Seg.TriangleID[0]), Mesh.GetTriangle(Seg.TriangleID[1]) };
+				int MatchedVertices = 0;
+				for (int T0SubIdx = 0; T0SubIdx < 3; T0SubIdx++)
+				{
+					FVector3d V0 = Mesh.GetVertex(Tri[0][T0SubIdx]);
+					for (int T1SubIdx = 0; T1SubIdx < 3; T1SubIdx++)
+					{
+						FVector3d V1 = Mesh.GetVertex(Tri[1][T1SubIdx]);
+						if (DistanceSquared(V0, V1) < FMathd::ZeroTolerance)
+						{
+							MatchedVertices++;
+							break;
+						}
+					}
+				}
+				// no shared vertices: treat as a real collision
+				// (TODO: only skip shared edges? will need to do something to avoid shared vertices becoming collisions)
+				if (MatchedVertices < 1)
+				{
+					bAnyUpdatesNeeded = true;
+					FUpdate& Update = Updates[CellIdx].Emplace_GetRef(Seg.TriangleID[0], Seg.TriangleID[1]);
+					for (int TriIdx = 0; TriIdx < 2; TriIdx++)
+					{
+						for (int VSubIdx = 0; VSubIdx < 3; VSubIdx++)
+						{
+							int VIdx = Tri[TriIdx][VSubIdx];
+							FVector3d P = Mesh.GetVertex(VIdx);
+							FIndex2i IDs(CellIdx, VIdx);
+							MeshesVertices.FindPointsInBall(P, FMathd::ZeroTolerance, [this, P](FIndex2i IDs)
+								{
+									FVector3d Pos = CellMeshes[IDs.A].AugMesh.GetVertex(IDs.B);
+									return DistanceSquared(P, Pos);
+								}, Update.IDs);
+						}
+					}
+				}
+			}
+		}
+		if (!bAnyUpdatesNeeded)
+		{
+			break;
+		}
+		for (int CellIdx = 0; CellIdx < NumCells; CellIdx++)
+		{
+			CellUnmoved[CellIdx] = true;
+		}
+
+		// todo: maybe can parallelize if movements are not applied until after?
+		for (int CellIdx = 0; CellIdx < NumCells; CellIdx++)
+		{
+			FDynamicMesh3& Mesh = CellMeshes[CellIdx].AugMesh;
+			TDynamicMeshVertexAttribute<double, 3>* OriginalPosns =
+				static_cast<TDynamicMeshVertexAttribute<double, 3>*>(Mesh.Attributes()->GetAttachedAttribute(OriginalPositionAttribute));
+			auto InterpVert = [&Mesh, &OriginalPosns](int VID, double t)
+			{
+				FVector3d OrigPos, NoisePos;
+				OriginalPosns->GetValue(VID, OrigPos);
+				NoisePos = Mesh.GetVertex(VID);
+				return UE::Geometry::Lerp(OrigPos, NoisePos, t);
+			};
+			auto InterpTri = [&Mesh, &InterpVert](int TID, double t)
+			{
+				FIndex3i TriVIDs = Mesh.GetTriangle(TID);
+				FTriangle3d Tri;
+				for (int i = 0; i < 3; i++)
+				{
+					Tri.V[i] = InterpVert(TriVIDs[i], t);
+				}
+				return Tri;
+			};
+			auto TestIntersection = [&InterpTri](int TIDA, int TIDB, double t)
+			{
+				FIntrTriangle3Triangle3d TriTri(InterpTri(TIDA, t), InterpTri(TIDB, t));
+				return TriTri.Find();
+			};
+			// resolve tri-tri intersections on this cell's mesh (moving associated verts on other meshes as needed also)
+			for (FUpdate& Update : Updates[CellIdx])
+			{
+				double tsafe = 0;
+				double tbad = 1;
+				if (!TestIntersection(Update.Tris.A, Update.Tris.B, tbad))
+				{
+					continue;
+				}
+				for (int SearchSteps = 0; SearchSteps < 4; SearchSteps++)
+				{
+					double tmid = (tsafe + tbad) * .5;
+					if (TestIntersection(Update.Tris.A, Update.Tris.B, tmid))
+					{
+						tbad = tmid;
+					}
+					else
+					{
+						tsafe = tmid;
+					}
+				}
+				CellUnmoved[CellIdx] = false;
+				for (FIndex2i IDs : Update.IDs)
+				{
+					FVector3d OldPos = CellMeshes[IDs.A].AugMesh.GetVertex(IDs.B);
+					FVector3d NewPos;
+					if (IDs.A == CellIdx)
+					{
+						NewPos = InterpVert(IDs.B, tsafe);
+						Mesh.SetVertex(IDs.B, NewPos);
+					}
+					else
+					{
+						CellUnmoved[IDs.A] = false;
+						FDynamicMesh3& OtherMesh = CellMeshes[IDs.A].AugMesh;
+						TDynamicMeshVertexAttribute<double, 3>* OtherOriginalPosns =
+							static_cast<TDynamicMeshVertexAttribute<double, 3>*>(OtherMesh.Attributes()->GetAttachedAttribute(OriginalPositionAttribute));
+						FVector3d OrigPos;
+						OtherOriginalPosns->GetValue(IDs.B, OrigPos);
+						NewPos = UE::Geometry::Lerp(OrigPos, OldPos, tsafe);
+						OtherMesh.SetVertex(IDs.B, NewPos);
+					}
+					MeshesVertices.UpdatePoint(IDs, OldPos, NewPos);
+				}
+			}
+		}
+	}
+
+	// clear "original position" attribute now that we have removed self-intersections
+	for (int CellIdx = 0; CellIdx < NumCells; CellIdx++)
+	{
+		FCellInfo& CellInfo = CellMeshes[CellIdx];
+		FDynamicMesh3& Mesh = CellInfo.AugMesh;
+		Mesh.Attributes()->RemoveAttribute(OriginalPositionAttribute);
+	}
+
+	// recompute UVs using new positions after noise was applied + fixed
+	TArray<FVector2f> PlaneMinUVs; PlaneMinUVs.Init(FVector2f(FMathf::MaxReal, FMathf::MaxReal), Cells.Planes.Num());
+	TArray<FFrame3d> PlaneFrames; PlaneFrames.Reserve(Cells.Planes.Num());
+	for (int PlaneIdx = 0; PlaneIdx < Cells.Planes.Num(); PlaneIdx++)
+	{
+		PlaneFrames.Emplace(Cells.Planes[PlaneIdx]);
+	}
+	// first pass to compute min UV for each plane
+	for (FCellInfo& CellInfo : CellMeshes)
+	{
+		FDynamicMesh3& Mesh = CellInfo.AugMesh;
+		FDynamicMeshMaterialAttribute* MaterialIDs = Mesh.Attributes()->GetMaterialID();
+
+		for (int TID : Mesh.TriangleIndicesItr())
+		{
+			int PlaneIdx = MaterialToPlane(MaterialIDs->GetValue(TID));
+			if (PlaneIdx > -1)
+			{
+				FIndex3i Tri = Mesh.GetTriangle(TID);
+				for (int Idx = 0; Idx < 3; Idx++)
+				{
+					FVector2f UV = FVector2f(PlaneFrames[PlaneIdx].ToPlaneUV(Mesh.GetVertex(Tri[Idx])));
+					FVector2f& MinUV = PlaneMinUVs[PlaneIdx];
+					MinUV.X = FMathf::Min(UV.X, MinUV.X);
+					MinUV.Y = FMathf::Min(UV.Y, MinUV.Y);
+				}
+			}
+		}
+	}
+	// second pass to actually set UVs
+	for (FCellInfo& CellInfo : CellMeshes)
+	{
+		FDynamicMesh3& Mesh = CellInfo.AugMesh;
+		FDynamicMeshMaterialAttribute* MaterialIDs = Mesh.Attributes()->GetMaterialID();
+
+		for (int TID : Mesh.TriangleIndicesItr())
+		{
+			int PlaneIdx = MaterialToPlane(MaterialIDs->GetValue(TID));
+			if (PlaneIdx > -1)
+			{
+				FIndex3i Tri = Mesh.GetTriangle(TID);
+				for (int Idx = 0; Idx < 3; Idx++)
+				{
+					FVector2f UV = ((FVector2f)PlaneFrames[PlaneIdx].ToPlaneUV(Mesh.GetVertex(Tri[Idx])) - PlaneMinUVs[PlaneIdx]) * GlobalUVScale;
+					Mesh.SetVertexUV(Tri[Idx], UV);
+				}
+			}
+		}
+	}
+}
+
+
+void FCellMeshes::CreateMeshesForSinglePlane(const FPlanarCells& Cells, const FAxisAlignedBox3d& DomainBounds, bool bNoise, double GlobalUVScale, double Grout, bool bOnlyGrout)
+{
+	bool bHasGrout = Grout > 0;
+
+	int MID = PlaneToMaterial(0);
+	FPlane Plane = Cells.Planes[0];
+
+	FFrame3d PlaneFrame(Plane);
+	FInterval1d ZRange;
+	FAxisAlignedBox2d XYRange;
+	for (int CornerIdx = 0; CornerIdx < 8; CornerIdx++)
+	{
+		FVector3d Corner = DomainBounds.GetCorner(CornerIdx);
+		XYRange.Contain(PlaneFrame.ToPlaneUV(Corner));
+		ZRange.Contain(Plane.PlaneDot(FVector(Corner)));
+	}
+	//if (FMathd::SignNonZero(ZRange.Min) == FMathd::SignNonZero(ZRange.Max))
+	//{
+	//	// TODO: early out for plane that doesn't even intersect the domain bounding box?
+	//}
+
+	FDynamicMesh3 PlaneMesh(true, true, true, false);
+	FVertexInfo PlaneVertInfo;
+	PlaneVertInfo.bHaveC = true;
+	PlaneVertInfo.bHaveUV = true;
+	PlaneVertInfo.bHaveN = true;
+	PlaneVertInfo.Color = FVector3f(1, 1, 1);
+	PlaneVertInfo.Normal = FVector3f(Plane.GetNormal());
+
+	for (int CornerIdx = 0; CornerIdx < 4; CornerIdx++)
+	{
+		PlaneVertInfo.Position = PlaneFrame.FromPlaneUV(XYRange.GetCorner(CornerIdx));
+		PlaneVertInfo.UV = FVector2f(XYRange.GetCorner(CornerIdx) - XYRange.Min) * GlobalUVScale;
+		PlaneMesh.AppendVertex(PlaneVertInfo);
+	}
+	PlaneMesh.AppendTriangle(0, 1, 2);
+	PlaneMesh.AppendTriangle(0, 2, 3);
+
+	if (bNoise)
+	{
+		double Spacing = GetSafeNoiseSpacing(XYRange.Area(), Cells.InternalSurfaceMaterials.NoiseSettings->PointSpacing);
+		RemeshForNoise(PlaneMesh, EEdgeRefineFlags::SplitsOnly, Spacing);
+		ApplyNoise(PlaneMesh, PlaneFrame.GetAxis(2), Cells.InternalSurfaceMaterials.NoiseSettings.GetValue(), true);
+		FMeshNormals::QuickComputeVertexNormals(PlaneMesh);
+	}
+	TArray<int> PlaneBoundary,  // loop of vertex IDs on the boundary of PlaneMesh (starting with vertex 0)
+		PlaneBoundaryCornerIndices; // indices of the corner vertices in the PlaneBoundary array
+	{
+		double Offset = ZRange.Max;
+		FMeshBoundaryLoops Boundary(&PlaneMesh);
+		checkSlow(Boundary.GetLoopCount() == 1);
+		int FirstIdx;
+		bool bFound = Boundary[0].Vertices.Find(0, FirstIdx);
+		checkSlow(bFound);
+		PlaneBoundary = Boundary[0].Vertices;
+		if (FirstIdx != 0)
+		{
+			Algo::Rotate(PlaneBoundary, FirstIdx);
+		}
+		checkSlow(PlaneBoundary[0] == 0);
+
+		PlaneBoundaryCornerIndices.Add(0);
+		int FoundIndices = 1;
+		for (int VIDIdx = 0; VIDIdx < PlaneBoundary.Num(); VIDIdx++)
+		{
+			int VID = PlaneBoundary[VIDIdx];
+			if (VID == FoundIndices)
+			{
+				FoundIndices++;
+				PlaneBoundaryCornerIndices.Add(VIDIdx);
+			}
+		}
+	}
+	FDynamicMesh3* Meshes[2];
+	if (!bOnlyGrout)
+	{
+		for (int Side = 0; Side < 2; Side++)
+		{
+			Meshes[Side] = &CellMeshes[Side].AugMesh;
+			*Meshes[Side] = PlaneMesh;
+			double Offset = ZRange.Max;
+			TArray<int> CapBoundary, CapBoundaryCornerIndices;
+
+			if (Side == 0)
+			{
+				Meshes[Side]->ReverseOrientation(true);
+				Offset = ZRange.Min;
+			}
+			PlaneVertInfo.Normal = FVector3f(Plane.GetNormal()) * (-1.0f + (float)Side * 2.0f);
+			FVector3d OffsetVec = FVector3d(Plane.GetNormal()) * Offset;
+
+			for (int CornerIdx = 0; CornerIdx < 4; CornerIdx++)
+			{
+				PlaneVertInfo.Position = Meshes[Side]->GetVertex(CornerIdx) + OffsetVec;
+				// UVs shouldn't matter for outer box vertices because they're outside of the domain by construction ...
+				CapBoundary.Add(Meshes[Side]->AppendVertex(PlaneVertInfo));
+				CapBoundaryCornerIndices.Add(CornerIdx);
+			}
+			int NewTris[2]{
+				Meshes[Side]->AppendTriangle(CapBoundary[0], CapBoundary[1], CapBoundary[2]),
+				Meshes[Side]->AppendTriangle(CapBoundary[0], CapBoundary[2], CapBoundary[3])
+			};
+			if (Side == 1)
+			{
+				Meshes[Side]->ReverseTriOrientation(NewTris[0]);
+				Meshes[Side]->ReverseTriOrientation(NewTris[1]);
+			}
+			FDynamicMeshEditor Editor(Meshes[Side]);
+			FDynamicMeshEditResult ResultOut;
+			Editor.StitchSparselyCorrespondedVertexLoops(PlaneBoundary, PlaneBoundaryCornerIndices, CapBoundary, CapBoundaryCornerIndices, ResultOut, Side == 0);
+		}
+	}
+	if (bHasGrout)
+	{
+		int GroutIdx = bOnlyGrout ? 0 : 2;
+		FDynamicMesh3* GroutMesh = &CellMeshes[GroutIdx].AugMesh;
+		FVector3d GroutOffset = (FVector3d)Plane.GetNormal() * (Grout * .5);
+		if (!bOnlyGrout)
+		{
+			for (int Side = 0; Side < 2; Side++)
+			{
+				// shift both sides out by Grout/2
+				MeshTransforms::Translate(*Meshes[Side], GroutOffset * (-1.0 + (double)Side * 2.0));
+			}
+		}
+
+		// make the center (grout) by stitching together two offset copies of PlaneMesh
+		*GroutMesh = PlaneMesh;
+		GroutMesh->ReverseOrientation(true);
+		MeshTransforms::Translate(*GroutMesh, GroutOffset);
+		FMeshIndexMappings IndexMaps;
+		FDynamicMeshEditor Editor(GroutMesh);
+		Editor.AppendMesh(&PlaneMesh, IndexMaps, [GroutOffset](int VID, const FVector3d& PosIn) {return PosIn - GroutOffset; });
+		TArray<int> AppendPlaneBoundary; AppendPlaneBoundary.Reserve(PlaneBoundary.Num());
+		TArray<int> RevBoundary = PlaneBoundary;
+		Algo::Reverse(RevBoundary);
+		for (int VID : RevBoundary)
+		{
+			AppendPlaneBoundary.Add(IndexMaps.GetNewVertex(VID));
+		}
+		FDynamicMeshEditResult ResultOut;
+		Editor.StitchVertexLoopsMinimal(RevBoundary, AppendPlaneBoundary, ResultOut);
+	}
+
+	// fix up custom attributes and material IDs for all meshes
+	for (int CellIdx = 0; CellIdx < CellMeshes.Num(); CellIdx++)
+	{
+		FDynamicMesh3& Mesh = CellMeshes[CellIdx].AugMesh;
+		// re-enable tangents and visibility attributes, since these are lost when we set the mesh to a copy of the plane mesh
+		AugmentedDynamicMesh::Augment(Mesh);
+
+		// Set all material IDs to the one plane's corresponding material ID
+		for (int TID : Mesh.TriangleIndicesItr())
+		{
+			Mesh.Attributes()->GetMaterialID()->SetNewValue(TID, MID);
+		}
+	}
+}
+
+
+void FDynamicMeshCollection::Init(const FGeometryCollection* Collection, const TArrayView<const int32>& TransformIndices, FTransform TransformCollection, bool bSaveIsolatedVertices)
+{
+	Meshes.Reset();
+	Bounds = FAxisAlignedBox3d::Empty();
+
+	for (int32 TransformIdx : TransformIndices)
+	{
+		if (Collection->Children[TransformIdx].Num() > 0)
+		{
+			// only store the meshes of leaf nodes
+			continue;
+		}
+
+		FTransform3d CollectionToLocal = FTransform3d(GeometryCollectionAlgo::GlobalMatrix(Collection->Transform, Collection->Parent, TransformIdx) * TransformCollection);
+
+		int32 AddedMeshIdx = Meshes.Add(new FMeshData);
+		FMeshData& MeshData = Meshes[AddedMeshIdx];
+		MeshData.TransformIndex = TransformIdx;
+		MeshData.ToCollection = FTransform(CollectionToLocal.Inverse());
+		FDynamicMesh3& Mesh = MeshData.AugMesh;
+
+		int32 GeometryIdx = Collection->TransformToGeometryIndex[TransformIdx];
+		Mesh.EnableAttributes();
+		Mesh.Attributes()->EnableMaterialID();
+
+		int32 VertexStart = Collection->VertexStart[GeometryIdx];
+		int32 VertexCount = Collection->VertexCount[GeometryIdx];
+		int32 FaceCount = Collection->FaceCount[GeometryIdx];
+
+		FVertexInfo VertexInfo;
+		VertexInfo.bHaveC = true;
+		VertexInfo.bHaveN = true;
+		VertexInfo.bHaveUV = true;
+		for (int32 Idx = VertexStart, N = VertexStart + VertexCount; Idx < N; Idx++)
+		{
+			VertexInfo.Position = CollectionToLocal.TransformPosition(FVector3d(Collection->Vertex[Idx]));
+			VertexInfo.UV = FVector2f(Collection->UV[Idx]);
+			VertexInfo.Color = FVector3f(Collection->Color[Idx]);
+			VertexInfo.Normal = (FVector3f)CollectionToLocal.TransformVectorNoScale(FVector3d(Collection->Normal[Idx]));
+			int VID = Mesh.AppendVertex(VertexInfo);
+			AugmentedDynamicMesh::SetTangent(Mesh, VID, VertexInfo.Normal,
+				(FVector3f)CollectionToLocal.TransformVectorNoScale(FVector3d(Collection->TangentU[Idx])),
+				(FVector3f)CollectionToLocal.TransformVectorNoScale(FVector3d(Collection->TangentV[Idx])));
+		}
+		FIntVector VertexOffset(VertexStart, VertexStart, VertexStart);
+		for (int32 Idx = Collection->FaceStart[GeometryIdx], N = Collection->FaceStart[GeometryIdx] + FaceCount; Idx < N; Idx++)
+		{
+			FIndex3i AddTri = FIndex3i(Collection->Indices[Idx] - VertexOffset);
+			int TID = Mesh.AppendTriangle(AddTri, 0);
+			if (TID == FDynamicMesh3::NonManifoldID)
+			{
+				// work around non-manifold triangles by copying the vertices
+				FIndex3i NewTri(-1, -1, -1);
+				for (int SubIdx = 0; SubIdx < 3; SubIdx++)
+				{
+					int NewVID = Mesh.AppendVertex(Mesh, AddTri[SubIdx]);
+					int32 SrcIdx = AddTri[SubIdx] + VertexStart;
+					AugmentedDynamicMesh::SetTangent(Mesh, NewVID,
+						Mesh.GetVertexNormal(NewVID), // TODO: we don't actually use the vertex normal; consider removing this arg from the function entirely
+						(FVector3f)CollectionToLocal.TransformVectorNoScale(FVector3d(Collection->TangentU[SrcIdx])),
+						(FVector3f)CollectionToLocal.TransformVectorNoScale(FVector3d(Collection->TangentV[SrcIdx])));
+					NewTri[SubIdx] = NewVID;
+				}
+				TID = Mesh.AppendTriangle(NewTri, 0);
+			}
+			if (TID < 0)
+			{
+				continue;
+			}
+			Mesh.Attributes()->GetMaterialID()->SetValue(TID, Collection->MaterialID[Idx]);
+			AugmentedDynamicMesh::SetVisibility(Mesh, TID, Collection->Visible[Idx]);
+			// note: material index doesn't need to be passed through; will be rebuilt by a call to reindex materials once the cut mesh is returned back to geometry collection format
+		}
+
+		if (!bSaveIsolatedVertices)
+		{
+			FDynamicMeshEditor Editor(&Mesh);
+			Editor.RemoveIsolatedVertices();
+		}
+
+		Bounds.Contain(Mesh.GetCachedBounds());
+
+		// TODO: build spatial data (add this after setting up mesh boolean path that can use it)
+		//MeshData.Spatial.SetMesh(&Mesh);
+	}
+}
+
+
+int32 FDynamicMeshCollection::CutWithMultiplePlanes(
+	const TArrayView<const FPlane>& Planes,
+	double Grout,
+	double CollisionSampleSpacing,
+	FGeometryCollection* Collection,
+	FInternalSurfaceMaterials& InternalSurfaceMaterials,
+	bool bSetDefaultInternalMaterialsFromCollection
+)
+{
+#if WITH_EDITOR
+	// Create progress indicator dialog
+	static const FText SlowTaskText = NSLOCTEXT("CutMultipleWithMultiplePlanes", "CutMultipleWithMultiplePlanesText", "Cutting geometry collection with plane(s)...");
+
+	FScopedSlowTask SlowTask(Planes.Num(), SlowTaskText);
+	SlowTask.MakeDialog();
+
+	// Declare progress shortcut lambdas
+	auto EnterProgressFrame = [&SlowTask](float Progress)
+	{
+		SlowTask.EnterProgressFrame(Progress);
+	};
+#else
+	auto EnterProgressFrame = [](float Progress) {};
+#endif
+
+	bool bHasGrout = Grout > 0;
+
+	if (bHasGrout)
+	{
+		// For multi-plane cuts with grout specifically, the easiest path seems to be:
+		// 1. Build the "grout" section of each plane
+		// 2. Take the union of all those grout sections as the grout mesh
+		// 3. Use the generic CutWithCellMeshes path, where that grout mesh is both the inner and outside cell mesh
+		//    (Note the outside cell mesh is subtracted, not intersected)
+		//    (Note this relies on island splitting to separate all the pieces afterwards.)
+		FCellMeshes GroutCells;
+		GroutCells.SetNumCells(2);
+		FDynamicMesh3& GroutMesh = GroutCells.CellMeshes[0].AugMesh;
+		FDynamicMeshEditor GroutAppender(&GroutMesh);
+		FMeshIndexMappings IndexMaps;
+		for (int32 PlaneIdx = 0; PlaneIdx < Planes.Num(); PlaneIdx++)
+		{
+			EnterProgressFrame(.5);
+			FPlanarCells PlaneCells(Planes[PlaneIdx]);
+			PlaneCells.InternalSurfaceMaterials = InternalSurfaceMaterials;
+			FCellMeshes PlaneGroutMesh;
+			PlaneGroutMesh.MakeOnlyPlanarGroutCell(PlaneCells, Bounds, Grout);
+			GroutAppender.AppendMesh(&PlaneGroutMesh.CellMeshes[0].AugMesh, IndexMaps);
+		}
+
+		EnterProgressFrame(Planes.Num() * .2);
+		FMeshSelfUnion GroutUnion(&GroutMesh);
+		GroutUnion.bSimplifyAlongNewEdges = true;
+		GroutUnion.bWeldSharedEdges = false;
+		GroutUnion.Compute();
+
+		EnterProgressFrame(Planes.Num() * .1);
+		// first mesh is the same as the second mesh, but will be subtracted b/c it's the "outside cell"
+		GroutCells.CellMeshes[1].AugMesh = GroutMesh;
+		GroutCells.OutsideCellIndex = 1;
+
+		EnterProgressFrame(Planes.Num() * .2);
+		TArray<TPair<int32, int32>> CellConnectivity;
+		CellConnectivity.Add(TPair<int32, int32>(0, -1));
+
+		return CutWithCellMeshes(InternalSurfaceMaterials, CellConnectivity, GroutCells, Collection, bSetDefaultInternalMaterialsFromCollection, CollisionSampleSpacing);
+	}
+
+	bool bHasProximity = Collection->HasAttribute("Proximity", FGeometryCollection::GeometryGroup);
+	TArray<TUniquePtr<FMeshData>> ToCut;
+	TArray<TUniquePtr<TPointHashGrid3d<int>>> VerticesHashes;
+	auto HashMeshVertices = [&VerticesHashes, &ToCut](int32 HashIdx)
+	{
+		FDynamicMesh3& Mesh = ToCut[HashIdx]->AugMesh;
+		if (HashIdx >= VerticesHashes.Num())
+		{
+			VerticesHashes.SetNum(HashIdx + 1);
+		}
+		if (VerticesHashes[HashIdx].IsValid())
+		{
+			return;
+		}
+		VerticesHashes[HashIdx] = MakeUnique<TPointHashGrid3d<int>>(FMathd::ZeroTolerance * 1000, -1);
+		TPointHashGrid3d<int>& Grid = *VerticesHashes[HashIdx].Get();
+		for (int VID : Mesh.VertexIndicesItr())
+		{
+			Grid.InsertPointUnsafe(VID, Mesh.GetVertex(VID));
+		}
+	};
+	auto ClearHash = [&VerticesHashes](int32 HashIdx)
+	{
+		if (HashIdx < VerticesHashes.Num())
+		{
+			VerticesHashes[HashIdx].Release();
+		}
+	};
+	auto IsNeighbor = [&ToCut, &VerticesHashes](int32 A, int32 B)
+	{
+		if (!ensure(A < ToCut.Num() && B < ToCut.Num() && A < VerticesHashes.Num() && B < VerticesHashes.Num()))
+		{
+			return false;
+		}
+		if (!ensure(VerticesHashes[A].IsValid() && VerticesHashes[B].IsValid()))
+		{
+			return false;
+		}
+		if (!ToCut[A]->AugMesh.GetCachedBounds().Intersects(ToCut[B]->AugMesh.GetCachedBounds()))
+		{
+			return false;
+		}
+		if (ToCut[A]->AugMesh.VertexCount() > ToCut[B]->AugMesh.VertexCount())
+		{
+			Swap(A, B);
+		}
+		FDynamicMesh3& RefMesh = ToCut[B]->AugMesh;
+		for (const FVector3d& V : ToCut[A]->AugMesh.VerticesItr())
+		{
+			TPair<int, double> Nearest = VerticesHashes[B]->FindNearestInRadius(V, FMathd::ZeroTolerance * 10, [&RefMesh, &V](int VID)
+				{
+					return DistanceSquared(RefMesh.GetVertex(VID), V);
+				});
+			if (Nearest.Key != -1)
+			{
+				return true;
+			}
+		}
+		return false;
+	};
+	// copy initial surfaces
+	for (FMeshData& MeshData : Meshes)
+	{
+		ToCut.Add(MakeUnique<FMeshData>(MeshData));
+	}
+	// track connections between meshes via their indices in the ToCut array
+	TMultiMap<int32, int32> Proximity;
+	auto ProxLink = [&Proximity](int32 A, int32 B)
+	{
+		Proximity.Add(A, B);
+		Proximity.Add(B, A);
+	};
+	auto ProxUnlink = [&Proximity](int32 A, int32 B)
+	{
+		Proximity.RemoveSingle(A, B);
+		Proximity.RemoveSingle(B, A);
+	};
+	for (int32 PlaneIdx = 0; PlaneIdx < Planes.Num(); PlaneIdx++)
+	{
+		EnterProgressFrame(1);
+		FPlanarCells PlaneCells(Planes[PlaneIdx]);
+		PlaneCells.InternalSurfaceMaterials = InternalSurfaceMaterials;
+		double OnePercentExtend = Bounds.MaxDim() * .01;
+		FCellMeshes CellMeshes(PlaneCells, Bounds, 0, OnePercentExtend, false);
+
+		// TODO: we could do these cuts in parallel (will takes some rework of the proximity and how results are added to the ToCut array)
+		for (int32 ToCutIdx = 0, ToCutNum = ToCut.Num(); ToCutIdx < ToCutNum; ToCutIdx++)
+		{
+			FMeshData& Surface = *ToCut[ToCutIdx];
+			int32 TransformIndex = Surface.TransformIndex;
+			FTransform ToCollection = Surface.ToCollection;
+
+			FAxisAlignedBox3d Box = Surface.AugMesh.GetCachedBounds();
+			if (InternalSurfaceMaterials.NoiseSettings.IsSet())
+			{
+				Box.Expand(InternalSurfaceMaterials.NoiseSettings->Amplitude);
+			}
+			if (!FMath::PlaneAABBIntersection(Planes[PlaneIdx], FBox(Box)))
+			{
+				continue;
+			}
+
+			TArray<TUniquePtr<FMeshData>> BoolResults;
+			for (int ResultIdx = 0; ResultIdx < 2; ResultIdx++)
+			{
+				BoolResults.Add(MakeUnique<FMeshData>());
+				BoolResults[ResultIdx]->TransformIndex = TransformIndex;
+				BoolResults[ResultIdx]->ToCollection = ToCollection;
+			}
+			check(CellMeshes.CellMeshes.Num() == 2);
+			bool bKeepResults = true;
+			for (int32 CellIdx = 0; CellIdx < 2; CellIdx++)
+			{
+				FCellMeshes::FCellInfo& Cell = CellMeshes.CellMeshes[CellIdx];
+
+				FMeshBoolean::EBooleanOp Op = FMeshBoolean::EBooleanOp::Intersect;
+				if (CellIdx == CellMeshes.OutsideCellIndex)
+				{
+					Op = FMeshBoolean::EBooleanOp::Difference;
+				}
+				FMeshBoolean Boolean(&Surface.AugMesh, &Cell.AugMesh, &BoolResults[CellIdx]->AugMesh, Op);
+				Boolean.bSimplifyAlongNewEdges = true;
+				Boolean.PreserveUVsOnlyForMesh = 0; // slight warping of the autogenerated cell UVs generally doesn't matter
+				Boolean.bWeldSharedEdges = false;
+				if (!Boolean.Compute())
+				{
+					// TODO: do something about failure cases?  e.g. try auto-filling small holes?
+					// note: failure cases won't be detected at all unless we weld edges,
+					//       which will require re-working how tangents are carried through
+				}
+				if (BoolResults[CellIdx]->AugMesh.TriangleCount() == 0)
+				{
+					bKeepResults = false;
+					break;
+				}
+			}
+
+			if (bKeepResults)
+			{
+				ToCut[ToCutIdx] = MoveTemp(BoolResults[0]);
+				int32 NewIdx = ToCut.Add(MoveTemp(BoolResults[1]));
+				// indices of all boolean result meshes (may be more than two due to splitting disconnected components)
+				TArray<int32, TInlineAllocator<4>> ResultIndices = { ToCutIdx, NewIdx };
+				// corresponding parent indices for each result mesh
+				TArray<int32, TInlineAllocator<4>> ParentIndices = { 0, 1 };
+				TArray<FDynamicMesh3> SplitMeshes;
+				for (int UnsplitIdx = 0; UnsplitIdx < 2; UnsplitIdx++)
+				{
+					if (SplitIslands(ToCut[ResultIndices[UnsplitIdx]]->AugMesh, SplitMeshes))
+					{
+						ToCut[ResultIndices[UnsplitIdx]]->AugMesh = SplitMeshes[0];
+						for (int32 Idx = 1; Idx < SplitMeshes.Num(); Idx++)
+						{
+							const FDynamicMesh3& Mesh = SplitMeshes[Idx];
+							ResultIndices.Add(ToCut.Add(MakeUnique<FMeshData>(Mesh, TransformIndex, ToCollection)));
+							ParentIndices.Add(UnsplitIdx);
+						}
+					}
+				}
+
+				// update proximity for neighbors of the original piece
+				if (bHasProximity)
+				{
+					ClearHash(ToCutIdx);
+					TArray<int32> Nbrs;
+					Proximity.MultiFind(ToCutIdx, Nbrs);
+					if (Nbrs.Num() > 0)
+					{
+						for (int32 ChangedMeshIdx : ResultIndices)
+						{
+							HashMeshVertices(ChangedMeshIdx);
+						}
+
+						for (int32 Nbr : Nbrs)
+						{
+							ProxUnlink(ToCutIdx, Nbr);
+							HashMeshVertices(Nbr);
+							for (int32 Idx = 0; Idx < ResultIndices.Num(); Idx++)
+							{
+								int32 ResultIdx = ResultIndices[Idx];
+								int32 OldIdx = Nbr;
+								if (IsNeighbor(ResultIdx, OldIdx))
+								{
+									ProxLink(ResultIdx, OldIdx);
+								}
+							}
+						}
+					}
+
+					if (ResultIndices.Num() == 2)
+					{
+						// add the connection between the two new pieces
+						ProxLink(ResultIndices[0], ResultIndices[1]);
+					}
+					else
+					{
+						if (Nbrs.Num() == 0)
+						{
+							for (int32 ChangedMeshIdx : ResultIndices)
+							{
+								HashMeshVertices(ChangedMeshIdx);
+							}
+						}
+						// check for connections between all pieces
+						for (int FirstIdx = 0; FirstIdx + 1 < ResultIndices.Num(); FirstIdx++)
+						{
+							int32 FirstParent = ParentIndices[FirstIdx];
+							for (int SecondIdx = FirstIdx + 1; SecondIdx < ResultIndices.Num(); SecondIdx++)
+							{
+								if (FirstParent == ParentIndices[SecondIdx])
+								{
+									// these pieces split from the same mesh *because* they were disconnected, so the pieces cannot be neighbors
+									continue;
+								}
+								if (IsNeighbor(ResultIndices[FirstIdx], ResultIndices[SecondIdx]))
+								{
+									ProxLink(ResultIndices[FirstIdx], ResultIndices[SecondIdx]);
+								}
+							}
+						}
+					}
+				}
+			} // iteration over meshes to cut
+		} // iteration over cutting planes
+	}
+
+	TMultiMap<int32, int32> ParentTransformToChildren;
+	for (int32 ToCutIdx = 0; ToCutIdx < ToCut.Num(); ToCutIdx++)
+	{
+		ParentTransformToChildren.Add(ToCut[ToCutIdx]->TransformIndex, ToCutIdx);
+	}
+
+	TArray<int32> ToCutIdxToGeometryIdx;  ToCutIdxToGeometryIdx.Init(-1, ToCut.Num());
+	TArray<int32> ToCutIndices;
+	int32 FirstCreatedIndex = -1;
+	for (FMeshData& MeshData : Meshes)
+	{
+		int32 GeometryIdx = Collection->TransformToGeometryIndex[MeshData.TransformIndex];
+		int32 InternalMaterialID = bSetDefaultInternalMaterialsFromCollection ? InternalSurfaceMaterials.GetDefaultMaterialIDForGeometry(*Collection, GeometryIdx) : InternalSurfaceMaterials.GlobalMaterialID;
+		ToCutIndices.Reset();
+		ParentTransformToChildren.MultiFind(MeshData.TransformIndex, ToCutIndices);
+
+		// if there's only one mesh here, i.e. it didn't get cut at all
+		if (ToCutIndices.Num() <= 1)
+		{
+			continue;
+		}
+
+		// hide old parent geometry
+		SetVisibility(*Collection, GeometryIdx, false);
+
+		// add newly created geometry as children
+		int32 SubPartIdx = 0;
+		for (int32 ToCutIdx : ToCutIndices)
+		{
+			FDynamicMesh3& Mesh = ToCut[ToCutIdx]->AugMesh;
+
+			FString BoneName = GetBoneName(*Collection, ToCut[ToCutIdx]->TransformIndex, SubPartIdx++);
+			int32 CreatedGeometryIdx = AppendToCollection(ToCut[ToCutIdx]->ToCollection, Mesh, CollisionSampleSpacing, ToCut[ToCutIdx]->TransformIndex, BoneName, *Collection, InternalMaterialID);
+			ToCutIdxToGeometryIdx[ToCutIdx] = CreatedGeometryIdx;
+			if (FirstCreatedIndex == -1)
+			{
+				FirstCreatedIndex = CreatedGeometryIdx;
+			}
+		}
+	}
+
+	// create proximity sets on geometry collection and populate using ToCut's Proximity multimap and the array ToCutIdxToGeometryIdx
+	if (bHasProximity)
+	{
+		TManagedArray<TSet<int32>>& GCProximity = Collection->GetAttribute<TSet<int32>>("Proximity", FGeometryCollection::GeometryGroup);
+		for (TPair<int32, int32> Link : Proximity)
+		{
+			GCProximity[ToCutIdxToGeometryIdx[Link.Key]].Add(ToCutIdxToGeometryIdx[Link.Value]);
+		}
+	}
+
+	return FirstCreatedIndex;
+}
+
+
+int32 FDynamicMeshCollection::CutWithCellMeshes(const FInternalSurfaceMaterials& InternalSurfaceMaterials, const TArray<TPair<int32, int32>>& CellConnectivity, FCellMeshes& CellMeshes, FGeometryCollection* Collection, bool bSetDefaultInternalMaterialsFromCollection, double CollisionSampleSpacing)
+{
+	// TODO: should we do these cuts in parallel, and the appends sequentially below?
+	int32 FirstIdx = -1;
+	int BadCount = 0;
+	bool bHasProximity = Collection->HasAttribute("Proximity", FGeometryCollection::GeometryGroup);
+	for (FMeshData& Surface : Meshes)
+	{
+		int32 GeometryIdx = Collection->TransformToGeometryIndex[Surface.TransformIndex];
+		TArray<TUniquePtr<FDynamicMesh3>> BooleanResults;  BooleanResults.SetNum(CellMeshes.CellMeshes.Num());
+		ParallelFor(CellMeshes.CellMeshes.Num(), [&BooleanResults, &CellMeshes, &Surface](int32 CellIdx)
+			{
+				FCellMeshes::FCellInfo& Cell = CellMeshes.CellMeshes[CellIdx];
+				if (Cell.AugMesh.GetCachedBounds().Intersects(Surface.AugMesh.GetCachedBounds()))
+				{
+					BooleanResults[CellIdx] = MakeUnique<FDynamicMesh3>();
+					FDynamicMesh3& AugBoolResult = *BooleanResults[CellIdx];
+
+					FMeshBoolean::EBooleanOp Op = FMeshBoolean::EBooleanOp::Intersect;
+					if (CellIdx == CellMeshes.OutsideCellIndex)
+					{
+						Op = FMeshBoolean::EBooleanOp::Difference;
+					}
+					FMeshBoolean Boolean(&Surface.AugMesh, &Cell.AugMesh, &AugBoolResult, Op);
+					Boolean.bSimplifyAlongNewEdges = true;
+					Boolean.PreserveUVsOnlyForMesh = 0; // slight warping of the autogenerated cell UVs generally doesn't matter
+					Boolean.bWeldSharedEdges = false;
+					if (!Boolean.Compute())
+					{
+						// TODO: do something about failure cases?  e.g. try auto-filling small holes?
+						// note: failure cases won't be detected at all unless we weld edges,
+						//       which will require re-working how tangents are carried through
+					}
+				}
+			}, EParallelForFlags::None);
+
+		int32 NonEmptyResults = 0;
+		for (const TUniquePtr<FDynamicMesh3>& AugBoolResult : BooleanResults)
+		{
+			if (AugBoolResult.IsValid() && AugBoolResult->TriangleCount() > 0)
+			{
+				NonEmptyResults++;
+			}
+		}
+
+		if (NonEmptyResults > 1) // only write to geometry collection if more than one result was non-empty
+		{
+			TSet<int32> PlanesInOutput;
+			TMultiMap<int32, int32> CellToGeometry;
+			TMap<int32, int32> GeometryToResultMesh;
+			int32 SubPartIndex = 0;
+			int32 InternalMaterialID = bSetDefaultInternalMaterialsFromCollection ? InternalSurfaceMaterials.GetDefaultMaterialIDForGeometry(*Collection, GeometryIdx) : InternalSurfaceMaterials.GlobalMaterialID;
+
+			for (int32 CellIdx = 0; CellIdx < CellMeshes.CellMeshes.Num(); CellIdx++)
+			{
+				if (BooleanResults[CellIdx].IsValid() && BooleanResults[CellIdx]->TriangleCount() > 0)
+				{
+					FDynamicMesh3& AugBoolResult = *BooleanResults[CellIdx];
+					for (int TID : AugBoolResult.TriangleIndicesItr())
+					{
+						int MID = AugBoolResult.Attributes()->GetMaterialID()->GetValue(TID);
+						int32 PlaneIdx = CellMeshes.MaterialToPlane(MID);
+						if (PlaneIdx >= 0)
+						{
+							PlanesInOutput.Add(PlaneIdx);
+						}
+					}
+					int32 CreatedGeometryIdx = -1;
+					TArray<FDynamicMesh3> Islands;
+					if (SplitIslands(AugBoolResult, Islands))
+					{
+						for (int32 i = 0; i < Islands.Num(); i++)
+						{
+							FDynamicMesh3& Island = Islands[i];
+							FString BoneName = GetBoneName(*Collection, Surface.TransformIndex, SubPartIndex++);
+							CreatedGeometryIdx = AppendToCollection(Surface.ToCollection, Island, CollisionSampleSpacing, Surface.TransformIndex, BoneName, *Collection, InternalMaterialID);
+							CellToGeometry.Add(CellIdx, CreatedGeometryIdx);
+							if (i > 0)
+							{
+								GeometryToResultMesh.Add(CreatedGeometryIdx, BooleanResults.Add(MakeUnique<FDynamicMesh3>(Island)));
+							}
+							else
+							{
+								*BooleanResults[CellIdx] = Island;
+								GeometryToResultMesh.Add(CreatedGeometryIdx, CellIdx);
+							}
+						}
+					}
+					else
+					{
+						FString BoneName = GetBoneName(*Collection, Surface.TransformIndex, SubPartIndex++);
+						CreatedGeometryIdx = AppendToCollection(Surface.ToCollection, AugBoolResult, CollisionSampleSpacing, Surface.TransformIndex, BoneName, *Collection, InternalMaterialID);
+						CellToGeometry.Add(CellIdx, CreatedGeometryIdx);
+						GeometryToResultMesh.Add(CreatedGeometryIdx, CellIdx);
+					}
+					if (FirstIdx == -1)
+					{
+						FirstIdx = CreatedGeometryIdx;
+					}
+				}
+			}
+			if (bHasProximity)
+			{
+				TManagedArray<TSet<int32>>& Proximity = Collection->GetAttribute<TSet<int32>>("Proximity", FGeometryCollection::GeometryGroup);
+				TArray<TUniquePtr<TPointHashGrid3d<int>>> VertexHashes;
+				auto MakeHash = [this, &VertexHashes, &BooleanResults](int GID)
+				{
+					if (GID >= VertexHashes.Num())
+					{
+						VertexHashes.SetNum(GID + 1);
+					}
+					if (!VertexHashes[GID].IsValid())
+					{
+						VertexHashes[GID] = MakeUnique<TPointHashGrid3d<int>>(FMathd::ZeroTolerance * 1000, -1);
+						FillVertexHash(*BooleanResults[GID], *VertexHashes[GID]);
+					}
+				};
+				for (int32 PlaneIdx : PlanesInOutput)
+				{
+					TPair<int32, int32> Cells = CellConnectivity[PlaneIdx];
+					int32 SecondCell = Cells.Value < 0 ? CellMeshes.OutsideCellIndex : Cells.Value;
+					if (SecondCell != -1)
+					{
+						TArray<int32, TInlineAllocator<4>> GeomA, GeomB;
+						CellToGeometry.MultiFind(Cells.Key, GeomA, false);
+						CellToGeometry.MultiFind(SecondCell, GeomB, false);
+						if (GeomA.Num() == 1 && GeomB.Num() == 1)
+						{
+							Proximity[GeomA[0]].Add(GeomB[0]);
+							Proximity[GeomB[0]].Add(GeomA[0]);
+						}
+						else if (GeomA.Num() >= 1 && GeomB.Num() >= 1) // at least one was split; need to re-check proximities
+						{
+							for (int GIDA : GeomA)
+							{
+								int MeshA = GeometryToResultMesh[GIDA];
+								MakeHash(MeshA);
+								for (int GIDB : GeomB)
+								{
+									int MeshB = GeometryToResultMesh[GIDB];
+									MakeHash(MeshB);
+									if (IsNeighboring(*BooleanResults[MeshA], *VertexHashes[MeshA], *BooleanResults[MeshB], *VertexHashes[MeshB]))
+									{
+										Proximity[GIDA].Add(GIDB);
+										Proximity[GIDB].Add(GIDA);
+									}
+								}
+							}
+						}
+					}
+				}
+			}
+			// turn off old geom visibility (preferred default behavior)
+			SetVisibility(*Collection, GeometryIdx, false);
+		}
+	}
+
+	return FirstIdx;
+}
+
+
+void FDynamicMeshCollection::FillVertexHash(const FDynamicMesh3& Mesh, TPointHashGrid3d<int>& VertHash)
+{
+	for (int VID : Mesh.VertexIndicesItr())
+	{
+		FVector3d V = Mesh.GetVertex(VID);
+		VertHash.InsertPointUnsafe(VID, V);
+	}
+}
+
+
+bool FDynamicMeshCollection::IsNeighboring(FDynamicMesh3* Mesh[2], const TPointHashGrid3d<int>* VertHash[2])
+{
+	if (!ensure(Mesh[0] && Mesh[1] && VertHash[0] && VertHash[1]))
+	{
+		return false;
+	}
+	if (!Mesh[0]->GetCachedBounds().Intersects(Mesh[1]->GetCachedBounds()))
+	{
+		return false;
+	}
+
+	int A = 0, B = 1;
+	if (Mesh[0]->VertexCount() > Mesh[1]->VertexCount())
+	{
+		Swap(A, B);
+	}
+	FDynamicMesh3& RefMesh = *Mesh[B];
+	for (const FVector3d& V : Mesh[A]->VerticesItr())
+	{
+		TPair<int, double> Nearest = VertHash[B]->FindNearestInRadius(V, FMathd::ZeroTolerance * 10, [&RefMesh, &V](int VID)
+			{
+				return DistanceSquared(RefMesh.GetVertex(VID), V);
+			});
+		if (Nearest.Key != -1)
+		{
+			return true;
+		}
+	}
+	return false;
+}
+
+// Split mesh into connected components, including implicit connections by co-located vertices
+bool FDynamicMeshCollection::SplitIslands(FDynamicMesh3& Source, TArray<FDynamicMesh3>& SeparatedMeshes)
+{
+	double SnapDistance = 1e-03;
+	TPointHashGrid3d<int> VertHash(SnapDistance * 10, -1);
+	FDisjointSet VertComponents(Source.MaxVertexID());
+	// Add Source vertices to hash & disjoint sets
+	TArray<int> Neighbors;
+	for (int VID : Source.VertexIndicesItr())
+	{
+		FVector3d Pt = Source.GetVertex(VID);
+		Neighbors.Reset();
+		VertHash.FindPointsInBall(Pt, SnapDistance, [&Source, Pt](int OtherVID) {return DistanceSquared(Pt, Source.GetVertex(OtherVID)); }, Neighbors);
+		for (int NbrVID : Neighbors)
+		{
+			VertComponents.UnionSequential(VID, NbrVID);
+		}
+		VertHash.InsertPointUnsafe(VID, Pt);
+	}
+	for (FIndex3i Tri : Source.TrianglesItr())
+	{
+		VertComponents.Union(Tri.A, Tri.B);
+		VertComponents.Union(Tri.B, Tri.C);
+		VertComponents.Union(Tri.C, Tri.A);
+	}
+
+	bool bWasSplit = FDynamicMeshEditor::SplitMesh(&Source, SeparatedMeshes, [&Source, &VertComponents](int TID)
+		{
+			return (int)VertComponents.Find(Source.GetTriangle(TID).A);
+		});
+
+	if (bWasSplit)
+	{
+		// disconnected components that are contained inside other components need to be re-merged
+		TMeshSpatialSort<FDynamicMesh3> SpatialSort(SeparatedMeshes);
+		SpatialSort.NestingMethod = TMeshSpatialSort<FDynamicMesh3>::ENestingMethod::InLargestParent;
+		SpatialSort.bOnlyNestNegativeVolumes = false;
+		SpatialSort.bOnlyParentPostiveVolumes = true;
+		SpatialSort.Compute();
+		TArray<bool> KeepMeshes; KeepMeshes.Init(true, SeparatedMeshes.Num());
+		for (TMeshSpatialSort<FDynamicMesh3>::FMeshNesting& Nest : SpatialSort.Nests)
+		{
+			FDynamicMeshEditor Editor(&SeparatedMeshes[Nest.OuterIndex]);
+			FMeshIndexMappings Mappings;
+			for (int Inner : Nest.InnerIndices)
+			{
+				Editor.AppendMesh(&SeparatedMeshes[Inner], Mappings);
+				KeepMeshes[Inner] = false;
+			}
+		}
+		for (int Idx = 0; Idx < SeparatedMeshes.Num(); Idx++)
+		{
+			if (!KeepMeshes[Idx])
+			{
+				SeparatedMeshes.RemoveAtSwap(Idx, 1, false);
+				KeepMeshes.RemoveAtSwap(Idx, 1, false);
+				Idx--;
+			}
+		}
+	}
+	return bWasSplit;
+}
+
+	
+void FDynamicMeshCollection::AddCollisionSamples(double CollisionSampleSpacing)
+{
+	for (int32 MeshIdx = 0; MeshIdx < Meshes.Num(); MeshIdx++)
+	{
+		AugmentedDynamicMesh::AddCollisionSamplesPerComponent(Meshes[MeshIdx].AugMesh, CollisionSampleSpacing);
+	}
+}
+
+
+// Update all geometry in a GeometryCollection w/ the meshes in the MeshCollection
+// Resizes the GeometryCollection as needed
+bool FDynamicMeshCollection::UpdateAllCollections(FGeometryCollection& Collection)
+{
+	bool bAllSucceeded = true;
+
+	int32 NumGeometry = Collection.NumElements(FGeometryCollection::GeometryGroup);
+	TArray<int32> NewFaceCounts, NewVertexCounts;
+	NewFaceCounts.SetNumUninitialized(NumGeometry);
+	NewVertexCounts.SetNumUninitialized(NumGeometry);
+	for (int32 GeomIdx = 0; GeomIdx < Collection.FaceCount.Num(); GeomIdx++)
+	{
+		NewFaceCounts[GeomIdx] = Collection.FaceCount[GeomIdx];
+		NewVertexCounts[GeomIdx] = Collection.VertexCount[GeomIdx];
+	}
+	for (int MeshIdx = 0; MeshIdx < Meshes.Num(); MeshIdx++)
+	{
+		FDynamicMeshCollection::FMeshData& MeshData = Meshes[MeshIdx];
+		int32 GeomIdx = Collection.TransformToGeometryIndex[MeshData.TransformIndex];
+		NewFaceCounts[GeomIdx] = MeshData.AugMesh.TriangleCount();
+		NewVertexCounts[GeomIdx] = MeshData.AugMesh.VertexCount();
+	}
+	GeometryCollectionAlgo::ResizeGeometries(&Collection, NewFaceCounts, NewVertexCounts);
+
+	for (int MeshIdx = 0; MeshIdx < Meshes.Num(); MeshIdx++)
+	{
+		FDynamicMeshCollection::FMeshData& MeshData = Meshes[MeshIdx];
+		FDynamicMesh3& Mesh = MeshData.AugMesh;
+		int32 GeometryIdx = Collection.TransformToGeometryIndex[MeshData.TransformIndex];
+		bool bSucceeded = UpdateCollection(MeshData.ToCollection, Mesh, GeometryIdx, Collection, -1);
+		bAllSucceeded &= bSucceeded;
+	}
+
+	return bAllSucceeded;
+}
+
+
+// Update an existing geometry in a collection w/ a new mesh (w/ the same number of faces and vertices!)
+bool FDynamicMeshCollection::UpdateCollection(const FTransform& ToCollection, FDynamicMesh3& Mesh, int32 GeometryIdx, FGeometryCollection& Output, int32 InternalMaterialID)
+{
+	if (!Mesh.IsCompact())
+	{
+		Mesh.CompactInPlace(nullptr);
+	}
+
+	int32 OldVertexCount = Output.VertexCount[GeometryIdx];
+	int32 OldTriangleCount = Output.FaceCount[GeometryIdx];
+
+	int32 NewVertexCount = Mesh.VertexCount();
+	int32 NewTriangleCount = Mesh.TriangleCount();
+
+	if (!ensure(OldVertexCount == NewVertexCount) || !ensure(OldTriangleCount == NewTriangleCount))
+	{
+		return false;
+	}
+
+	int32 VerticesStart = Output.VertexStart[GeometryIdx];
+	int32 FacesStart = Output.FaceStart[GeometryIdx];
+	int32 TransformIdx = Output.TransformIndex[GeometryIdx];
+
+	for (int32 VID = 0; VID < Mesh.MaxVertexID(); VID++)
+	{
+		checkSlow(Mesh.IsVertex(VID)); // mesh is compact
+		int32 CopyToIdx = VerticesStart + VID;
+		Output.Vertex[CopyToIdx] = ToCollection.TransformPosition(FVector(Mesh.GetVertex(VID)));
+		Output.Normal[CopyToIdx] = ToCollection.TransformVectorNoScale(FVector(Mesh.GetVertexNormal(VID)));
+		Output.UV[CopyToIdx] = FVector2D(Mesh.GetVertexUV(VID));
+		FVector3f TangentU, TangentV;
+		AugmentedDynamicMesh::GetTangent(Mesh, VID, TangentU, TangentV);
+		Output.TangentU[CopyToIdx] = ToCollection.TransformVectorNoScale(FVector(TangentU));
+		Output.TangentV[CopyToIdx] = ToCollection.TransformVectorNoScale(FVector(TangentV));
+		Output.Color[CopyToIdx] = FVector(Mesh.GetVertexColor(VID));
+
+		// Bone map is set based on the transform of the new geometry
+		Output.BoneMap[CopyToIdx] = TransformIdx;
+	}
+
+	FIntVector VertexStartOffset(VerticesStart);
+	for (int32 TID = 0; TID < Mesh.MaxTriangleID(); TID++)
+	{
+		checkSlow(Mesh.IsTriangle(TID));
+		int32 CopyToIdx = FacesStart + TID;
+		Output.Visible[CopyToIdx] = AugmentedDynamicMesh::GetVisibility(Mesh, TID);
+		int MaterialID = Mesh.Attributes()->GetMaterialID()->GetValue(TID);
+		Output.MaterialID[CopyToIdx] = MaterialID < 0 ? InternalMaterialID : MaterialID;
+		Output.Indices[CopyToIdx] = FIntVector(Mesh.GetTriangle(TID)) + VertexStartOffset;
+	}
+
+	if (Output.BoundingBox.Num())
+	{
+		Output.BoundingBox[GeometryIdx].Init();
+		for (int32 Idx = VerticesStart; Idx < VerticesStart + Output.VertexCount[GeometryIdx]; ++Idx)
+		{
+			Output.BoundingBox[GeometryIdx] += Output.Vertex[Idx];
+		}
+	}
+
+	return true;
+}
+
+
+int32 FDynamicMeshCollection::AppendToCollection(const FTransform& ToCollection, FDynamicMesh3& Mesh, double CollisionSampleSpacing, int32 TransformParent, FString BoneName, FGeometryCollection& Output, int32 InternalMaterialID)
+{
+	if (Mesh.TriangleCount() == 0)
+	{
+		return -1;
+	}
+
+	if (!Mesh.IsCompact())
+	{
+		Mesh.CompactInPlace(nullptr);
+	}
+
+	if (CollisionSampleSpacing > 0)
+	{
+		AugmentedDynamicMesh::AddCollisionSamplesPerComponent(Mesh, CollisionSampleSpacing);
+	}
+
+	int32 NewGeometryStartIdx = Output.FaceStart.Num();
+	int32 OriginalVertexNum = Output.Vertex.Num();
+	int32 OriginalFaceNum = Output.Indices.Num();
+
+	int32 GeometryIdx = Output.AddElements(1, FGeometryCollection::GeometryGroup);
+	int32 TransformIdx = Output.AddElements(1, FGeometryCollection::TransformGroup);
+
+	int32 NumTriangles = Mesh.TriangleCount();
+	int32 NumVertices = Mesh.VertexCount();
+	check(NumTriangles > 0);
+	check(Mesh.IsCompact());
+	Output.FaceCount[GeometryIdx] = NumTriangles;
+	Output.FaceStart[GeometryIdx] = OriginalFaceNum;
+	Output.VertexCount[GeometryIdx] = NumVertices;
+	Output.VertexStart[GeometryIdx] = OriginalVertexNum;
+	Output.TransformIndex[GeometryIdx] = TransformIdx;
+	Output.TransformToGeometryIndex[TransformIdx] = GeometryIdx;
+	if (TransformParent > -1)
+	{
+		Output.BoneName[TransformIdx] = BoneName;
+		Output.BoneColor[TransformIdx] = Output.BoneColor[TransformParent];
+		Output.Parent[TransformIdx] = TransformParent;
+		Output.Children[TransformParent].Add(TransformIdx);
+		Output.SimulationType[TransformParent] = FGeometryCollection::ESimulationTypes::FST_Clustered;
+	}
+	Output.Transform[TransformIdx] = FTransform::Identity;
+	Output.SimulationType[TransformIdx] = FGeometryCollection::ESimulationTypes::FST_Rigid;
+
+	int32 FacesStart = Output.AddElements(NumTriangles, FGeometryCollection::FacesGroup);
+	int32 VerticesStart = Output.AddElements(NumVertices, FGeometryCollection::VerticesGroup);
+
+	for (int32 VID = 0; VID < Mesh.MaxVertexID(); VID++)
+	{
+		checkSlow(Mesh.IsVertex(VID)); // mesh is compact
+		int32 CopyToIdx = VerticesStart + VID;
+		Output.Vertex[CopyToIdx] = ToCollection.TransformPosition(FVector(Mesh.GetVertex(VID)));
+		Output.Normal[CopyToIdx] = ToCollection.TransformVectorNoScale(FVector(Mesh.GetVertexNormal(VID)));
+		Output.UV[CopyToIdx] = FVector2D(Mesh.GetVertexUV(VID));
+		FVector3f TangentU, TangentV;
+		AugmentedDynamicMesh::GetTangent(Mesh, VID, TangentU, TangentV);
+		Output.TangentU[CopyToIdx] = ToCollection.TransformVectorNoScale(FVector(TangentU));
+		Output.TangentV[CopyToIdx] = ToCollection.TransformVectorNoScale(FVector(TangentV));
+		Output.Color[CopyToIdx] = FVector(Mesh.GetVertexColor(VID));
+
+		// Bone map is set based on the transform of the new geometry
+		Output.BoneMap[CopyToIdx] = TransformIdx;
+	}
+
+	FIntVector VertexStartOffset(VerticesStart);
+	for (int32 TID = 0; TID < Mesh.MaxTriangleID(); TID++)
+	{
+		checkSlow(Mesh.IsTriangle(TID));
+		int32 CopyToIdx = FacesStart + TID;
+		Output.Visible[CopyToIdx] = AugmentedDynamicMesh::GetVisibility(Mesh, TID);
+		int MaterialID = Mesh.Attributes()->GetMaterialID()->GetValue(TID);
+		Output.MaterialID[CopyToIdx] = MaterialID < 0 ? InternalMaterialID : MaterialID;
+		Output.Indices[CopyToIdx] = FIntVector(Mesh.GetTriangle(TID)) + VertexStartOffset;
+	}
+
+	if (Output.BoundingBox.Num())
+	{
+		Output.BoundingBox[GeometryIdx].Init();
+		for (int32 Idx = OriginalVertexNum; Idx < Output.Vertex.Num(); ++Idx)
+		{
+			Output.BoundingBox[GeometryIdx] += Output.Vertex[Idx];
+		}
+	}
+
+	return GeometryIdx;
+}
+
+
+}} // namespace UE::PlanarCut
\ No newline at end of file
diff --git a/Engine/Plugins/Experimental/PlanarCutPlugin/Source/PlanarCut/Private/GeometryMeshConversion.h b/Engine/Plugins/Experimental/PlanarCutPlugin/Source/PlanarCut/Private/GeometryMeshConversion.h
new file mode 100644
--- /dev/null
+++ b/Engine/Plugins/Experimental/PlanarCutPlugin/Source/PlanarCut/Private/GeometryMeshConversion.h
@@ -0,0 +1,238 @@
+// Copyright Epic Games, Inc. All Rights Reserved.
+#pragma once
+#include "CoreMinimal.h"
+
+#include "DynamicMesh3.h"
+#include "DynamicMeshEditor.h"
+#include "DynamicMeshAABBTree3.h"
+#include "Selections/MeshConnectedComponents.h"
+#include "MeshTransforms.h"
+#include "MeshConstraints.h"
+#include "GeometryCollection/GeometryCollection.h"
+#include "Spatial/PointHashGrid3.h"
+
+struct FPlanarCells;
+struct FInternalSurfaceMaterials;
+struct FNoiseSettings;
+
+namespace UE
+{
+namespace PlanarCut
+{
+
+/**
+ * Add attributes necessary for a dynamic mesh to represent geometry from an FGeometryCollection
+ */
+void PLANARCUT_API SetGeometryCollectionAttributes(UE::Geometry::FDynamicMesh3& Mesh);
+
+// functions for DynamicMesh3 meshes that have FGeometryCollection attributes set
+namespace AugmentedDynamicMesh
+{
+	void PLANARCUT_API SetVisibility(UE::Geometry::FDynamicMesh3& Mesh, int TID, bool bIsVisible);
+	void PLANARCUT_API SetTangent(UE::Geometry::FDynamicMesh3& Mesh, int VID, FVector3f Normal, FVector3f TangentU, FVector3f TangentV);
+	void PLANARCUT_API GetTangent(const UE::Geometry::FDynamicMesh3& Mesh, int VID, FVector3f& U, FVector3f& V);
+	void PLANARCUT_API InitializeOverlayToPerVertexUVs(UE::Geometry::FDynamicMesh3& Mesh);
+	void PLANARCUT_API InitializeOverlayToPerVertexTangents(UE::Geometry::FDynamicMesh3& Mesh);
+	void PLANARCUT_API ComputeTangents(UE::Geometry::FDynamicMesh3& Mesh, bool bOnlyOddMaterials, const TArrayView<const int32>& WhichMaterials, bool bRecomputeNormals = true);
+	void PLANARCUT_API AddCollisionSamplesPerComponent(UE::Geometry::FDynamicMesh3& Mesh, double Spacing);
+}
+
+
+/**
+ * Dynamic mesh representation of cutting cells, to be used to fracture a mesh
+ */
+struct PLANARCUT_API FCellMeshes
+{
+	struct FCellInfo
+	{
+		UE::Geometry::FDynamicMesh3 AugMesh;
+
+		FCellInfo()
+		{
+			SetGeometryCollectionAttributes(AugMesh);
+		}
+
+		// TODO: compute spatial in advance?  (only useful if we rework mesh booleans to support it)
+		//FDynamicMeshAABBTree3 Spatial;
+	};
+
+	TIndirectArray<FCellInfo> CellMeshes;
+	int32 OutsideCellIndex = -1;
+
+	// Noise Offsets, to randomize where perlin noise is sampled
+	FVector NoiseOffsetX;
+	FVector NoiseOffsetY;
+	FVector NoiseOffsetZ;
+
+	void SetNumCells(int32 NumMeshes)
+	{
+		CellMeshes.Reset();
+		for (int32 Idx = 0; Idx < NumMeshes; Idx++)
+		{
+			CellMeshes.Add(new FCellInfo);
+		}
+	}
+
+	FCellMeshes()
+	{
+		InitEmpty();
+	}
+
+	FCellMeshes(const FPlanarCells& Cells, UE::Geometry::FAxisAlignedBox3d DomainBounds, double Grout, double ExtendDomain, bool bIncludeOutsideCell);
+
+	FCellMeshes(UE::Geometry::FDynamicMesh3& SingleCutter, const FInternalSurfaceMaterials& Materials, TOptional<FTransform> Transform);
+
+	// Special function to just make the "grout" part of the planar mesh cells
+	// Used to make the multi-plane cuts with grout easier to implement
+	void MakeOnlyPlanarGroutCell(const FPlanarCells& Cells, UE::Geometry::FAxisAlignedBox3d DomainBounds, double Grout);
+
+	void RemeshForNoise(UE::Geometry::FDynamicMesh3& Mesh, UE::Geometry::EEdgeRefineFlags EdgeFlags, double TargetEdgeLen);
+
+	float OctaveNoise(const FVector& V, const FNoiseSettings& Settings);
+
+	FVector NoiseVector(const FVector& Pos, const FNoiseSettings& Settings);
+
+	FVector3d NoiseDisplacement(const FVector3d& Pos, const FNoiseSettings& Settings);
+
+	void ApplyNoise(UE::Geometry::FDynamicMesh3& Mesh, FVector3d Normal, const FNoiseSettings& Settings, bool bProjectBoundariesToNormal = false);
+
+	/**
+	 * Convert plane index to material ID
+	 * @return material ID encoding the source plane into a triangle mesh
+	 */
+	int PlaneToMaterial(int Plane)
+	{
+		return -(Plane + 1);
+	}
+
+	/**
+	 * Convert material ID to plane index
+	 * @return index of source plane for triangle, or -1 if no such plane
+	 */
+	int MaterialToPlane(int MaterialID)
+	{
+		return MaterialID >= 0 ? -1 : -(MaterialID + 1);
+	}
+
+	void InitEmpty()
+	{
+		NoiseOffsetX = FMath::VRand() * 100;
+		NoiseOffsetY = FMath::VRand() * 100;
+		NoiseOffsetZ = FMath::VRand() * 100;
+		OutsideCellIndex = -1;
+	}
+
+	void Init(const FPlanarCells& Cells, UE::Geometry::FAxisAlignedBox3d DomainBounds, double Grout, double ExtendDomain, bool bIncludeOutsideCell);
+
+	void ApplyGeneralGrout(double Grout);
+
+	void AppendMesh(UE::Geometry::FDynamicMesh3& Base, UE::Geometry::FDynamicMesh3& ToAppend, bool bFlipped);
+
+private:
+
+	void CreateMeshesForBoundedPlanesWithoutNoise(int NumCells, const FPlanarCells& Cells, const UE::Geometry::FAxisAlignedBox3d& DomainBounds, bool bNoise, double GlobalUVScale);
+
+	// Approximately calculate a "safe" spacing that would not require the remesher to create more than a million new vertices
+	double GetSafeNoiseSpacing(float SurfaceArea, float TargetSpacing);
+
+	void CreateMeshesForBoundedPlanesWithNoise(int NumCells, const FPlanarCells& Cells, const UE::Geometry::FAxisAlignedBox3d& DomainBounds, bool bNoise, double GlobalUVScale);
+	void CreateMeshesForSinglePlane(const FPlanarCells& Cells, const UE::Geometry::FAxisAlignedBox3d& DomainBounds, bool bNoise, double GlobalUVScale, double Grout, bool bOnlyGrout);
+};
+
+
+// Holds Geometry from an FGeometryCollection in an FDynamicMesh3 representation, and convert both directions
+// Also supports cutting geometry with FCellMeshes
+struct PLANARCUT_API FDynamicMeshCollection
+{
+	struct FMeshData
+	{
+		UE::Geometry::FDynamicMesh3 AugMesh;
+
+		// FDynamicMeshAABBTree3 Spatial; // TODO: maybe refactor mesh booleans to allow version where caller provides spatial data; it's computed every boolean now
+		// FTransform3d Transform; // TODO: maybe pretransform the data to a space that is good for cutting; refactor mesh boolean so there is an option to have it not transform input
+		int32 TransformIndex; // where the mesh was from in the geometry collection
+		FTransform ToCollection; // transform that need be applied to go back to the local space of the geometry collection
+
+		FMeshData()
+		{
+			SetGeometryCollectionAttributes(AugMesh);
+		}
+
+		FMeshData(const UE::Geometry::FDynamicMesh3& Mesh, int32 TransformIndex, FTransform ToCollection) : AugMesh(Mesh), TransformIndex(TransformIndex), ToCollection(ToCollection)
+		{}
+	};
+	TIndirectArray<FMeshData> Meshes;
+	UE::Geometry::FAxisAlignedBox3d Bounds;
+
+	FDynamicMeshCollection(const FGeometryCollection* Collection, const TArrayView<const int32>& TransformIndices, FTransform TransformCollection, bool bSaveIsolatedVertices = false)
+	{
+		Init(Collection, TransformIndices, TransformCollection, bSaveIsolatedVertices);
+	}
+
+	void Init(const FGeometryCollection* Collection, const TArrayView<const int32>& TransformIndices, FTransform TransformCollection, bool bSaveIsolatedVertices = false);
+
+	int32 CutWithMultiplePlanes(
+		const TArrayView<const FPlane>& Planes,
+		double Grout,
+		double CollisionSampleSpacing,
+		FGeometryCollection* Collection,
+		FInternalSurfaceMaterials& InternalSurfaceMaterials,
+		bool bSetDefaultInternalMaterialsFromCollection
+	);
+
+	/**
+	 * Cut collection meshes with cell meshes, and append results to a geometry collection
+	 *
+	 * @param InternalSurfaceMaterials Internal material info (used for material ID)
+	 * @param CellConnectivity The connectivity between cells: PlaneTag -> The two cells separated by triangles with this tag
+	 * @param CellsMeshes Meshed versions of the cells, with noise and material properties baked in
+	 * @param Collection Results will be stored in this
+	 * @param bSetDefaultInternalMaterialsFromCollection If true, set internal materials to the most common external material + 1, following a convenient artist convention
+	 * @return Index of the first created geometry
+	 */
+	int32 CutWithCellMeshes(const FInternalSurfaceMaterials& InternalSurfaceMaterials, const TArray<TPair<int32, int32>>& CellConnectivity, FCellMeshes& CellMeshes, FGeometryCollection* Collection, bool bSetDefaultInternalMaterialsFromCollection, double CollisionSampleSpacing);
+
+	static void SetVisibility(FGeometryCollection& Collection, int32 GeometryIdx, bool bVisible)
+	{
+		int32 FaceEnd = Collection.FaceCount[GeometryIdx] + Collection.FaceStart[GeometryIdx];
+		for (int32 FaceIdx = Collection.FaceStart[GeometryIdx]; FaceIdx < FaceEnd; FaceIdx++)
+		{
+			Collection.Visible[FaceIdx] = bVisible;
+		}
+	}
+
+	// Split mesh into connected components, including implicit connections by co-located vertices
+	bool SplitIslands(UE::Geometry::FDynamicMesh3& Source, TArray<UE::Geometry::FDynamicMesh3>& SeparatedMeshes);
+
+	FString GetBoneName(FGeometryCollection& Output, int TransformParent, int SubPartIndex)
+	{
+		return Output.BoneName[TransformParent] + "_" + FString::FromInt(SubPartIndex);
+	}
+
+	void AddCollisionSamples(double CollisionSampleSpacing);
+
+	// Update all geometry in a GeometryCollection w/ the meshes in the MeshCollection
+	// Resizes the GeometryCollection as needed
+	bool UpdateAllCollections(FGeometryCollection& Collection);
+
+	static int32 AppendToCollection(const FTransform& ToCollection, UE::Geometry::FDynamicMesh3& Mesh, double CollisionSampleSpacing, int32 TransformParent, FString BoneName, FGeometryCollection& Output, int32 InternalMaterialID);
+
+private:
+
+	// Update an existing geometry in a collection w/ a new mesh (w/ the same number of faces and vertices!)
+	static bool UpdateCollection(const FTransform& ToCollection, UE::Geometry::FDynamicMesh3& Mesh, int32 GeometryIdx, FGeometryCollection& Output, int32 InternalMaterialID);
+
+	void FillVertexHash(const UE::Geometry::FDynamicMesh3& Mesh, UE::Geometry::TPointHashGrid3d<int>& VertHash);
+
+	bool IsNeighboring(UE::Geometry::FDynamicMesh3& MeshA, const UE::Geometry::TPointHashGrid3d<int>& VertHashA, UE::Geometry::FDynamicMesh3& MeshB, const UE::Geometry::TPointHashGrid3d<int>& VertHashB)
+	{
+		UE::Geometry::FDynamicMesh3* Mesh[2]{ &MeshA, &MeshB };
+		const UE::Geometry::TPointHashGrid3d<int>* VertHash[2]{ &VertHashA, &VertHashB };
+		return IsNeighboring(Mesh, VertHash);
+	}
+
+	bool IsNeighboring(UE::Geometry::FDynamicMesh3* Mesh[2], const UE::Geometry::TPointHashGrid3d<int>* VertHash[2]);
+};
+
+
+}} // namespace UE::PlanarCut
\ No newline at end of file
diff --git a/Engine/Plugins/Experimental/PlanarCutPlugin/Source/PlanarCut/Private/PlanarCut.cpp b/Engine/Plugins/Experimental/PlanarCutPlugin/Source/PlanarCut/Private/PlanarCut.cpp
--- a/Engine/Plugins/Experimental/PlanarCutPlugin/Source/PlanarCut/Private/PlanarCut.cpp
+++ b/Engine/Plugins/Experimental/PlanarCutPlugin/Source/PlanarCut/Private/PlanarCut.cpp
@@ -40,6 +40,8 @@
 
 #include "Algo/Rotate.h"
 
+#include "GeometryMeshConversion.h"
+
 #if WITH_EDITOR
 #include "Misc/ScopedSlowTask.h"
 #endif
@@ -48,2257 +50,7 @@
 using namespace UE::Geometry;
 
 
-
-namespace UE
-{
-	namespace PlanarCutInternals
-	{
-		// functions to setup geometry collection attributes on dynamic meshes
-		namespace AugmentDynamicMesh
-		{
-			FName TangentUAttribName = "TangentUAttrib";
-			FName TangentVAttribName = "TangentVAttrib";
-			FName VisibleAttribName = "VisibleAttrib";
-
-			void Augment(FDynamicMesh3& Mesh)
-			{
-				Mesh.EnableVertexColors(FVector3f(1, 1, 1));
-				Mesh.EnableVertexNormals(FVector3f::UnitZ());
-				Mesh.EnableVertexUVs(FVector2f(0, 0));
-				Mesh.EnableAttributes();
-				Mesh.Attributes()->EnableMaterialID();
-				ensure(Mesh.Attributes()->NumAttachedAttributes() == 0);
-				Mesh.Attributes()->AttachAttribute(TangentUAttribName, new TDynamicMeshVertexAttribute<float, 3>(&Mesh));
-				Mesh.Attributes()->AttachAttribute(TangentVAttribName, new TDynamicMeshVertexAttribute<float, 3>(&Mesh));
-				TDynamicMeshScalarTriangleAttribute<bool>* VisAttrib = new TDynamicMeshScalarTriangleAttribute<bool>(&Mesh);
-				VisAttrib->Initialize(true);
-				Mesh.Attributes()->AttachAttribute(VisibleAttribName, VisAttrib);
-			}
-
-			bool IsAugmented(const FDynamicMesh3& Mesh)
-			{
-				return Mesh.HasAttributes()
-					&& Mesh.Attributes()->NumAttachedAttributes() == 3
-					&& Mesh.Attributes()->HasAttachedAttribute(TangentUAttribName)
-					&& Mesh.Attributes()->HasAttachedAttribute(TangentVAttribName)
-					&& Mesh.Attributes()->HasAttachedAttribute(VisibleAttribName)
-					&& Mesh.Attributes()->HasMaterialID()
-					&& Mesh.HasVertexColors()
-					&& Mesh.HasVertexNormals()
-					&& Mesh.HasVertexUVs();
-			}
-
-
-			void SetDefaultAttributes(FDynamicMesh3& Mesh, const FInternalSurfaceMaterials& Materials)
-			{
-				checkSlow(IsAugmented(Mesh));
-				TDynamicMeshVertexAttribute<float, 3>* Us =
-					static_cast<TDynamicMeshVertexAttribute<float, 3>*>(Mesh.Attributes()->GetAttachedAttribute(TangentUAttribName));
-				TDynamicMeshVertexAttribute<float, 3>* Vs =
-					static_cast<TDynamicMeshVertexAttribute<float, 3>*>(Mesh.Attributes()->GetAttachedAttribute(TangentVAttribName));
-
-				for (int VID : Mesh.VertexIndicesItr())
-				{
-					FVector3f N = Mesh.GetVertexNormal(VID);
-					FVector3f U, V;
-					VectorUtil::MakePerpVectors(N, U, V);
-					Us->SetValue(VID, U);
-					Vs->SetValue(VID, V);
-				}
-
-				TDynamicMeshScalarTriangleAttribute<bool>* Visible =
-					static_cast<TDynamicMeshScalarTriangleAttribute<bool>*>(Mesh.Attributes()->GetAttachedAttribute(VisibleAttribName));
-				for (int TID : Mesh.TriangleIndicesItr())
-				{
-					Visible->SetNewValue(TID, Materials.bGlobalVisibility);
-				}
-			}
-
-			void SetVisibility(FDynamicMesh3& Mesh, int TID, bool bIsVisible)
-			{
-				checkSlow(IsAugmented(Mesh));
-				TDynamicMeshScalarTriangleAttribute<bool>* Visible =
-					static_cast<TDynamicMeshScalarTriangleAttribute<bool>*>(Mesh.Attributes()->GetAttachedAttribute(VisibleAttribName));
-				Visible->SetValue(TID, bIsVisible);
-			}
-
-			bool GetVisibility(const FDynamicMesh3& Mesh, int TID)
-			{
-				checkSlow(IsAugmented(Mesh));
-				TDynamicMeshScalarTriangleAttribute<bool>* Visible =
-					static_cast<TDynamicMeshScalarTriangleAttribute<bool>*>(Mesh.Attributes()->GetAttachedAttribute(VisibleAttribName));
-				return Visible->GetValue(TID);
-			}
-
-			void SetTangent(FDynamicMesh3& Mesh, int VID, FVector3f Normal, FVector3f TangentU, FVector3f TangentV)
-			{
-				checkSlow(IsAugmented(Mesh));
-				TDynamicMeshVertexAttribute<float, 3>* Us =
-					static_cast<TDynamicMeshVertexAttribute<float, 3>*>(Mesh.Attributes()->GetAttachedAttribute(TangentUAttribName));
-				TDynamicMeshVertexAttribute<float, 3>* Vs =
-					static_cast<TDynamicMeshVertexAttribute<float, 3>*>(Mesh.Attributes()->GetAttachedAttribute(TangentVAttribName));
-				Us->SetValue(VID, TangentU);
-				Vs->SetValue(VID, TangentV);
-			}
-
-			void GetTangent(const FDynamicMesh3& Mesh, int VID, FVector3f& U, FVector3f& V)
-			{
-				checkSlow(IsAugmented(Mesh));
-				TDynamicMeshVertexAttribute<float, 3>* Us =
-					static_cast<TDynamicMeshVertexAttribute<float, 3>*>(Mesh.Attributes()->GetAttachedAttribute(TangentUAttribName));
-				TDynamicMeshVertexAttribute<float, 3>* Vs =
-					static_cast<TDynamicMeshVertexAttribute<float, 3>*>(Mesh.Attributes()->GetAttachedAttribute(TangentVAttribName));
-				FVector3f Normal = Mesh.GetVertexNormal(VID);
-				Us->GetValue(VID, U);
-				Vs->GetValue(VID, V);
-			}
-
-			void InitializeOverlayToPerVertexUVs(FDynamicMesh3& Mesh)
-			{
-				FDynamicMeshUVOverlay* UVs = Mesh.Attributes()->PrimaryUV();
-				UVs->ClearElements();
-				TArray<int> VertToUVMap;
-				VertToUVMap.SetNumUninitialized(Mesh.MaxVertexID());
-				for (int VID : Mesh.VertexIndicesItr())
-				{
-					FVector2f UV = Mesh.GetVertexUV(VID);
-					int UVID = UVs->AppendElement(UV);
-					VertToUVMap[VID] = UVID;
-				}
-				for (int TID : Mesh.TriangleIndicesItr())
-				{
-					FIndex3i Tri = Mesh.GetTriangle(TID);
-					Tri.A = VertToUVMap[Tri.A];
-					Tri.B = VertToUVMap[Tri.B];
-					Tri.C = VertToUVMap[Tri.C];
-					UVs->SetTriangle(TID, Tri);
-				}
-			}
-
-			void InitializeOverlayToPerVertexTangents(FDynamicMesh3& Mesh)
-			{
-				Mesh.Attributes()->EnableTangents();
-				FDynamicMeshNormalOverlay* TangentOverlays[2] = { Mesh.Attributes()->PrimaryTangents(), Mesh.Attributes()->PrimaryBiTangents() };
-				TangentOverlays[0]->ClearElements();
-				TangentOverlays[1]->ClearElements();
-				TArray<int> VertToTangentMap;
-				VertToTangentMap.SetNumUninitialized(Mesh.MaxVertexID());
-				for (int VID : Mesh.VertexIndicesItr())
-				{
-					FVector3f Tangents[2];
-					UE::PlanarCutInternals::AugmentDynamicMesh::GetTangent(Mesh, VID, Tangents[0], Tangents[1]);
-					int TID = TangentOverlays[0]->AppendElement(&Tangents[0].X);
-					int TID2 = TangentOverlays[1]->AppendElement(&Tangents[1].X);
-					check(TID == TID2);
-					VertToTangentMap[VID] = TID;
-				}
-				for (int TID : Mesh.TriangleIndicesItr())
-				{
-					FIndex3i Tri = Mesh.GetTriangle(TID);
-					Tri.A = VertToTangentMap[Tri.A];
-					Tri.B = VertToTangentMap[Tri.B];
-					Tri.C = VertToTangentMap[Tri.C];
-					TangentOverlays[0]->SetTriangle(TID, Tri);
-					TangentOverlays[1]->SetTriangle(TID, Tri);
-				}
-			}
-
-			void ComputeTangents(FDynamicMesh3& Mesh, bool bOnlyOddMaterials, const TArrayView<const int32>& WhichMaterials, bool bRecomputeNormals = true)
-			{
-				FDynamicMeshNormalOverlay* Normals = Mesh.Attributes()->PrimaryNormals();
-				FMeshNormals::InitializeOverlayToPerVertexNormals(Normals, !bRecomputeNormals);
-				if (bRecomputeNormals)
-				{
-					FMeshNormals::QuickRecomputeOverlayNormals(Mesh);
-				}
-
-				// Copy per-vertex UVs to a UV overlay, because that's what the tangents code uses
-				// (TODO: consider making a tangent computation path that uses vertex normals / UVs)
-				InitializeOverlayToPerVertexUVs(Mesh);
-				FDynamicMeshUVOverlay* UVs = Mesh.Attributes()->PrimaryUV();
-
-				FComputeTangentsOptions Options;
-				Options.bAngleWeighted = true;
-				Options.bAveraged = true;
-				FMeshTangentsf Tangents(&Mesh);
-				Tangents.ComputeTriVertexTangents(Normals, UVs, Options);
-
-				const TArray<FVector3<float>>& TanU = Tangents.GetTangents();
-				const TArray<FVector3<float>>& TanV = Tangents.GetBitangents();
-				FDynamicMeshMaterialAttribute* MaterialIDs = Mesh.Attributes()->GetMaterialID();
-				for (int TID : Mesh.TriangleIndicesItr())
-				{
-					int MaterialID = MaterialIDs->GetValue(TID);
-					if (bOnlyOddMaterials && (MaterialID % 2) == 0)
-					{
-						continue;
-					}
-					else if (WhichMaterials.Contains(MaterialID))
-					{
-						continue;
-					}
-					
-					int TanIdxBase = TID * 3;
-					FIndex3i Tri = Mesh.GetTriangle(TID);
-					for (int Idx = 0; Idx < 3; Idx++)
-					{
-						int VID = Tri[Idx];
-						int TanIdx = TanIdxBase + Idx;
-						UE::PlanarCutInternals::AugmentDynamicMesh::SetTangent(Mesh, VID, Mesh.GetVertexNormal(VID), TanU[TanIdx], TanV[TanIdx]);
-					}
-				}
-			}
-
-			// per component sampling is a rough heuristic to avoid doing geodesic distance but still get points on a 'thin' slice
-			void AddCollisionSamplesPerComponent(FDynamicMesh3& Mesh, double Spacing)
-			{
-				checkSlow(IsAugmented(Mesh));
-				FMeshConnectedComponents Components(&Mesh);
-				// TODO: if/when we switch to merged edges representation, pass a predicate here based on whether there's a normal seam ?
-				Components.FindConnectedTriangles();
-				TArray<TPointHashGrid3d<int>> KnownSamples;  KnownSamples.Reserve(Components.Num());
-				for (int ComponentIdx = 0; ComponentIdx < Components.Num(); ComponentIdx++)
-				{
-					KnownSamples.Emplace(.5 * Spacing / FMathd::InvSqrt3, -1);
-				}
-
-				TArray<int> AlreadySeen; AlreadySeen.Init(-1, Mesh.MaxVertexID());
-				for (int ComponentIdx = 0; ComponentIdx < Components.Num(); ComponentIdx++)
-				{
-					FMeshConnectedComponents::FComponent& Component = Components.GetComponent(ComponentIdx);
-					for (int TID : Component.Indices)
-					{
-						FIndex3i Tri = Mesh.GetTriangle(TID);
-						for (int SubIdx = 0; SubIdx < 3; SubIdx++)
-						{
-							int VID = Tri[SubIdx];
-							if (AlreadySeen[VID] != ComponentIdx)
-							{
-								AlreadySeen[VID] = ComponentIdx;
-								KnownSamples[ComponentIdx].InsertPointUnsafe(VID, Mesh.GetVertex(VID));
-							}
-						}
-					}
-				}
-				AlreadySeen.Empty();
-
-				double SpacingThreshSq = Spacing * Spacing; // if points are more than Spacing apart, consider adding a new point between them
-				for (int ComponentIdx = 0; ComponentIdx < Components.Num(); ComponentIdx++)
-				{
-					FMeshConnectedComponents::FComponent& Component = Components.GetComponent(ComponentIdx);
-					for (int TID : Component.Indices)
-					{
-						FIndex3i TriVIDs = Mesh.GetTriangle(TID);
-						FTriangle3d Triangle;
-						Mesh.GetTriVertices(TID, Triangle.V[0], Triangle.V[1], Triangle.V[2]);
-						double EdgeLensSq[3];
-						int MaxEdgeIdx = 0;
-						double MaxEdgeLenSq = 0;
-						for (int i = 2, j = 0; j < 3; i = j++)
-						{
-							double EdgeLenSq = DistanceSquared(Triangle.V[i], Triangle.V[j]);
-							if (EdgeLenSq > MaxEdgeLenSq)
-							{
-								MaxEdgeIdx = i;
-								MaxEdgeLenSq = EdgeLenSq;
-							}
-							EdgeLensSq[i] = EdgeLenSq;
-						}
-						// if we found a too-long edge, we can try sampling the tri
-						if (MaxEdgeLenSq > SpacingThreshSq)
-						{
-							FVector3f Normal = (FVector3f)VectorUtil::Normal(Triangle.V[0], Triangle.V[1], Triangle.V[2]);
-
-							// Pick number of samples based on the longest edge
-							double LongEdgeLen = FMathd::Sqrt(MaxEdgeLenSq);
-							int Divisions = FMathd::Floor(LongEdgeLen / Spacing);
-							double Factor = 1.0 / double(Divisions + 1);
-							int SecondEdgeIdx = (MaxEdgeIdx + 1) % 3;
-							int ThirdEdgeIdx = (MaxEdgeIdx + 2) % 3;
-							// Sample along the two longest edges first, then interpolate these samples
-							int SecondLongestEdgeIdx = SecondEdgeIdx;
-							if (EdgeLensSq[SecondEdgeIdx] < EdgeLensSq[ThirdEdgeIdx])
-							{
-								SecondLongestEdgeIdx = ThirdEdgeIdx;
-							}
-							int SecondLongestSecondEdgeIdx = (SecondLongestEdgeIdx + 1) % 3;
-							for (int DivI = 0; DivI < Divisions; DivI++)
-							{
-								double Along = (DivI + 1) * Factor;
-								FVector3d E1Bary(0, 0, 0), E2Bary(0, 0, 0);
-								E1Bary[MaxEdgeIdx] = Along;
-								E1Bary[SecondEdgeIdx] = 1 - Along;
-								E2Bary[SecondLongestEdgeIdx] = 1 - Along;
-								E2Bary[SecondLongestSecondEdgeIdx] = Along;
-
-								// Choose number of samples between the two edge points based on their distance
-								double AcrossDist = Triangle.BarycentricPoint(E1Bary).Distance(Triangle.BarycentricPoint(E2Bary));
-								int DivisionsAcross = FMathd::Ceil(AcrossDist / Spacing);
-								double FactorAcross = 1.0 / double(DivisionsAcross + 1);
-								for (int DivJ = 0; DivJ < DivisionsAcross; DivJ++)
-								{
-									double AlongAcross = (DivJ + 1) * FactorAcross;
-									FVector3d Bary = UE::Geometry::Lerp(E1Bary, E2Bary, AlongAcross);
-									FVector3d SamplePos = Triangle.BarycentricPoint(Bary);
-									if (!KnownSamples[ComponentIdx].IsCellEmptyUnsafe(SamplePos)) // fast early out; def. have pt within radius
-									{
-										continue;
-									}
-									TPair<int, double> VIDDist = KnownSamples[ComponentIdx].FindNearestInRadius(SamplePos, Spacing * .5, [&Mesh, SamplePos](int VID)
-										{
-											return DistanceSquared(Mesh.GetVertex(VID), SamplePos);
-										});
-									// No point within radius Spacing/2 -> Add a new sample
-									if (VIDDist.Key == -1)
-									{
-										// no point within radius; can add a sample here
-										FVertexInfo Info(SamplePos, Normal);
-
-										int AddedVID = Mesh.AppendVertex(Info);
-										KnownSamples[ComponentIdx].InsertPointUnsafe(AddedVID, SamplePos);
-									}
-								}
-							}
-						}
-					}
-				}
-			}
-		}
-	}
-}
-
-
-struct FCellMeshes
-{
-	struct FCellInfo
-	{
-		FDynamicMesh3 AugMesh;
-
-		FCellInfo()
-		{
-			UE::PlanarCutInternals::AugmentDynamicMesh::Augment(AugMesh);
-		}
-
-		// TODO: compute spatial in advance?  (only useful if we rework mesh booleans to support it)
-		//FDynamicMeshAABBTree3 Spatial;
-	};
-	TIndirectArray<FCellInfo> CellMeshes;
-	int32 OutsideCellIndex = -1;
-
-	// Noise Offsets, to randomize where perlin noise is sampled
-	FVector NoiseOffsetX;
-	FVector NoiseOffsetY;
-	FVector NoiseOffsetZ;
-
-	void SetNumCells(int32 NumMeshes)
-	{
-		CellMeshes.Reset();
-		for (int32 Idx = 0; Idx < NumMeshes; Idx++)
-		{
-			CellMeshes.Add(new FCellInfo);
-		}
-	}
-
-	FCellMeshes()
-	{
-		InitEmpty();
-	}
-	
-	FCellMeshes(const FPlanarCells& Cells, FAxisAlignedBox3d DomainBounds, double Grout, double ExtendDomain, bool bIncludeOutsideCell)
-	{
-		Init(Cells, DomainBounds, Grout, ExtendDomain, bIncludeOutsideCell);
-	}
-
-	FCellMeshes(FDynamicMesh3& SingleCutter, const FInternalSurfaceMaterials& Materials, TOptional<FTransform> Transform)
-	{
-		SetNumCells(2);
-
-		if (Transform.IsSet())
-		{
-			MeshTransforms::ApplyTransform(SingleCutter, UE::Geometry::FTransform3d(Transform.GetValue()));
-		}
-
-		// Mesh should already be augmented
-		if (!ensure(UE::PlanarCutInternals::AugmentDynamicMesh::IsAugmented(SingleCutter)))
-		{
-			UE::PlanarCutInternals::AugmentDynamicMesh::Augment(SingleCutter);
-		}
-
-		CellMeshes[0].AugMesh = SingleCutter;
-
-		// first mesh is the same as the second mesh, but will be subtracted b/c it's the "outside cell"
-		// TODO: special case this logic so we don't have to hold two copies of the exact same mesh!
-		CellMeshes[1].AugMesh = CellMeshes[0].AugMesh;
-		OutsideCellIndex = 1;
-
-	}
-	
-	// Special function to just make the "grout" part of the planar mesh cells
-	// Used to make the multi-plane cuts with grout easier to implement
-	void MakeOnlyPlanarGroutCell(const FPlanarCells& Cells, FAxisAlignedBox3d DomainBounds, double Grout)
-	{
-		CellMeshes.Reset();
-
-		if (!ensure(Grout > 0) || !ensure(Cells.IsInfinitePlane()))
-		{
-			return;
-		}
-		
-		float GlobalUVScale = Cells.InternalSurfaceMaterials.GlobalUVScale;
-		if (!ensure(GlobalUVScale > 0))
-		{
-			GlobalUVScale = 1;
-		}
-
-		SetNumCells(1);
-
-		bool bNoise = Cells.InternalSurfaceMaterials.NoiseSettings.IsSet();
-
-		double ExtendDomain = bNoise ? Cells.InternalSurfaceMaterials.NoiseSettings->Amplitude : 0;
-		DomainBounds.Expand(ExtendDomain);
-		
-		CreateMeshesForSinglePlane(Cells, DomainBounds, bNoise, GlobalUVScale, Grout, true);
-
-		for (FCellInfo& CellInfo : CellMeshes)
-		{
-			UE::PlanarCutInternals::AugmentDynamicMesh::SetDefaultAttributes(CellInfo.AugMesh, Cells.InternalSurfaceMaterials);
-		}
-	}
-
-	void RemeshForNoise(FDynamicMesh3& Mesh, EEdgeRefineFlags EdgeFlags, double TargetEdgeLen)
-	{
-		FQueueRemesher Remesh(&Mesh);
-		Remesh.bPreventNormalFlips = true;
-		FMeshConstraints Constraints;
-
-		FMeshBoundaryLoops Boundary(&Mesh);
-		int LoopCount = Boundary.GetLoopCount();
-		if (!ensureMsgf(LoopCount == 1, TEXT("Expected to remesh a patch with a single boundary but found %d boundary loops"), LoopCount))
-		{
-			if (LoopCount == 0)
-			{
-				return;
-			}
-		}
-
-		for (int VID : Mesh.VertexIndicesItr())
-		{
-			FVertexConstraint FullyConstrain(true, false, VID);
-			Constraints.SetOrUpdateVertexConstraint(VID, FullyConstrain);
-		}
-
-
-		FEdgeConstraint EdgeConstraint(EdgeFlags);
-		for (int EID : Boundary[0].Edges)
-		{
-			Constraints.SetOrUpdateEdgeConstraint(EID, EdgeConstraint);
-		}
-		Remesh.SetExternalConstraints(Constraints);
-		Remesh.SetTargetEdgeLength(TargetEdgeLen);
-		Remesh.Precompute();
-		Remesh.FastestRemesh();
-	}
-
-	float OctaveNoise(const FVector& V, const FNoiseSettings& Settings)
-	{
-		int32 Octaves = Settings.Octaves;
-		float NoiseValue = 0;
-		float OctaveScale = 1;
-		for (int32 Octave = 0; Octave < Octaves; Octave++, OctaveScale *= 2)
-		{
-			NoiseValue += FMath::PerlinNoise3D(V * OctaveScale) / OctaveScale;
-		}
-		return NoiseValue;
-	}
-
-	FVector NoiseVector(const FVector& Pos, const FNoiseSettings& Settings)
-	{
-		float Frequency = Settings.Frequency;
-		FVector Base = Pos * Frequency;
-		return FVector(
-			OctaveNoise(Base + NoiseOffsetX, Settings),
-			OctaveNoise(Base + NoiseOffsetY, Settings),
-			OctaveNoise(Base + NoiseOffsetZ, Settings)
-		) * Settings.Amplitude;
-	}
-
-	FVector3d NoiseDisplacement(const FVector3d& Pos, const FNoiseSettings& Settings)
-	{
-		FVector P = FVector(Pos);
-		return FVector3d(NoiseVector(P, Settings));
-	}
-
-	void ApplyNoise(FDynamicMesh3& Mesh, FVector3d Normal, const FNoiseSettings& Settings, bool bProjectBoundariesToNormal = false)
-	{
-		double Amplitude = (double)Settings.Amplitude;
-		double Frequency = (double)Settings.Frequency;
-		int32 Octaves = Settings.Octaves;
-		FVector3d Z = Normal * Amplitude;
-
-		for (int VID : Mesh.VertexIndicesItr())
-		{
-			FVector3d Pos = Mesh.GetVertex(VID);
-			FVector3d Displacement = NoiseDisplacement(Pos, Settings);
-			if (bProjectBoundariesToNormal || !Mesh.IsBoundaryVertex(VID))
-			{
-				// project displacement onto the normal direction
-				Displacement = Normal * Displacement.Dot(Normal);
-			}
-
-			
-			Mesh.SetVertex(VID, Pos + Displacement);
-		}
-	}
-
-	/**
-	 * Convert plane index to material ID
-	 * @return material ID encoding the source plane into a triangle mesh
-	 */
-	int PlaneToMaterial(int Plane)
-	{
-		return -(Plane+1);
-	}
-
-	/**
-	 * Convert material ID to plane index
-	 * @return index of source plane for triangle, or -1 if no such plane
-	 */
-	int MaterialToPlane(int MaterialID)
-	{
-		return MaterialID >= 0 ? -1 : -(MaterialID+1);
-	}
-
-	void InitEmpty()
-	{
-		NoiseOffsetX = FMath::VRand() * 100;
-		NoiseOffsetY = FMath::VRand() * 100;
-		NoiseOffsetZ = FMath::VRand() * 100;
-		OutsideCellIndex = -1;
-	}
-
-	void Init(const FPlanarCells& Cells, FAxisAlignedBox3d DomainBounds, double Grout, double ExtendDomain, bool bIncludeOutsideCell)
-	{
-		InitEmpty();
-
-		float GlobalUVScale = Cells.InternalSurfaceMaterials.GlobalUVScale;
-		if (!ensure(GlobalUVScale > 0))
-		{
-			GlobalUVScale = 1;
-		}
-
-		int NumCells = Cells.NumCells;
-		bool bHasGroutCell = Grout > 0;
-		if (bIncludeOutsideCell && !Cells.IsInfinitePlane())
-		{
-			OutsideCellIndex = NumCells;
-			NumCells++;
-		}
-
-		SetNumCells(NumCells);
-
-		bool bNoise = Cells.InternalSurfaceMaterials.NoiseSettings.IsSet();
-		if (bNoise)
-		{
-			ExtendDomain += Cells.InternalSurfaceMaterials.NoiseSettings->Amplitude;
-		}
-		DomainBounds.Expand(ExtendDomain);
-
-		// special handling for the infinite plane case; we need to adapt this to be a closed volume
-		if (Cells.IsInfinitePlane())
-		{
-			CreateMeshesForSinglePlane(Cells, DomainBounds, bNoise, GlobalUVScale, Grout, false);
-		}
-		else
-		{
-			if (!bNoise) // bounded cells w/ no noise
-			{
-				CreateMeshesForBoundedPlanesWithoutNoise(NumCells, Cells, DomainBounds, bNoise, GlobalUVScale);
-			}
-			else // bounded cells with noise -- make each boundary plane separately so we can remesh them w/ noise vertices
-			{
-				CreateMeshesForBoundedPlanesWithNoise(NumCells, Cells, DomainBounds, bNoise, GlobalUVScale);
-			}
-			ApplyGeneralGrout(Grout);
-		}
-		
-		// TODO: self-union on cells when it makes sense to do so (for non-single-plane inputs w/ high noise or possible untracked adjacencies)
-		/*for (FCellInfo& CellInfo : CellMeshes)
-		{
-			FMeshSelfUnion SelfUnion(&CellInfo.AugMesh);
-			// TODO: need to have an option in SelfUnion to not weld edges
-			SelfUnion.Compute();
-		}*/
-
-		for (FCellInfo& CellInfo : CellMeshes)
-		{
-			UE::PlanarCutInternals::AugmentDynamicMesh::SetDefaultAttributes(CellInfo.AugMesh, Cells.InternalSurfaceMaterials);
-		}
-	}
-
-	void ApplyGeneralGrout(double Grout)
-	{
-		if (Grout <= 0)
-		{
-			return;
-		}
-
-		// apply grout to all cells
-		for (int MeshIdx = 0; MeshIdx < CellMeshes.Num(); MeshIdx++)
-		{
-			if (MeshIdx == OutsideCellIndex)
-			{
-				continue;
-			}
-
-			FDynamicMesh3& Mesh = CellMeshes[MeshIdx].AugMesh;
-			// TODO: scale from mesh center of mass instead of the vertex centroid?
-			FVector3d VertexCentroid(0, 0, 0);
-			for (FVector3d V : Mesh.VerticesItr())
-			{
-				VertexCentroid += V;
-			}
-			VertexCentroid /= (double)Mesh.VertexCount();
-			FAxisAlignedBox3d Bounds = Mesh.GetCachedBounds();
-			double BoundsSize = Bounds.MaxDim();
-			// currently just scale the meshes down so they leave half-a-grout worth of space on their longest axis
-			// or delete the mesh if it's so small that that would require a negative scale
-			// TODO: consider instead computing a true offset mesh
-			//  (note that we don't currently have a good UV-preserving+sharp-edge-preserving way to do that)
-			double ScaleFactor = (BoundsSize - Grout*.5) / BoundsSize;
-			if (ScaleFactor < FMathd::ZeroTolerance * 1000)
-			{
-				// if the grout scale factor would be ~zero or negative, just clear the mesh instead
-				Mesh.Clear();
-				UE::PlanarCutInternals::AugmentDynamicMesh::Augment(Mesh);
-			}
-			else
-			{
-				MeshTransforms::Scale(Mesh, FVector3d::One() * ScaleFactor, VertexCentroid);
-			}
-		}
-
-		// create outside cell (if there is room for it) by appending all the other meshes
-		if (OutsideCellIndex != -1)
-		{
-			FDynamicMesh3& OutsideMesh = CellMeshes[OutsideCellIndex].AugMesh;
-			OutsideMesh.Clear();
-			UE::PlanarCutInternals::AugmentDynamicMesh::Augment(OutsideMesh);
-			FDynamicMeshEditor OutsideMeshEditor(&OutsideMesh);
-			for (int MeshIdx = 0; MeshIdx < CellMeshes.Num(); MeshIdx++)
-			{
-				if (MeshIdx == OutsideCellIndex)
-				{
-					continue;
-				}
-				FMeshIndexMappings IndexMaps;
-				OutsideMeshEditor.AppendMesh(&CellMeshes[MeshIdx].AugMesh, IndexMaps);
-			}
-		}
-	}
-
-	void AppendMesh(FDynamicMesh3& Base, FDynamicMesh3& ToAppend, bool bFlipped)
-	{
-		FDynamicMeshEditor Editor(&Base);
-		FMeshIndexMappings Mapping;
-		Editor.AppendMesh(&ToAppend, Mapping);
-		if (bFlipped)
-		{
-			for (int TID : ToAppend.TriangleIndicesItr())
-			{
-				Base.ReverseTriOrientation(Mapping.GetNewTriangle(TID));
-			}
-			for (int VID : ToAppend.VertexIndicesItr())
-			{
-				int BaseVID = Mapping.GetNewVertex(VID);
-				Base.SetVertexNormal(BaseVID, -Base.GetVertexNormal(BaseVID));
-			}
-		}
-	}
-private:
-	void CreateMeshesForBoundedPlanesWithoutNoise(int NumCells, const FPlanarCells& Cells, const FAxisAlignedBox3d& DomainBounds, bool bNoise, double GlobalUVScale)
-	{
-		for (int32 PlaneIdx = 0; PlaneIdx < Cells.PlaneCells.Num(); PlaneIdx++)
-		{
-			const TPair<int32, int32>& CellPair = Cells.PlaneCells[PlaneIdx];
-			FDynamicMesh3* Meshes[2]{ &CellMeshes[CellPair.Key].AugMesh, nullptr };
-			int32 OtherCell = CellPair.Value < 0 ? OutsideCellIndex : CellPair.Value;
-			int NumMeshes = OtherCell < 0 ? 1 : 2;
-			if (NumMeshes == 2)
-			{
-				Meshes[1] = &CellMeshes[OtherCell].AugMesh;
-			}
-
-			const TArray<int>& PlaneBoundary = Cells.PlaneBoundaries[PlaneIdx];
-			FVector3f Normal(Cells.Planes[PlaneIdx].GetNormal());
-			FFrame3d PlaneFrame(Cells.Planes[PlaneIdx]);
-			FVertexInfo PlaneVertInfo;
-			PlaneVertInfo.bHaveC = true;
-			PlaneVertInfo.bHaveUV = true;
-			PlaneVertInfo.bHaveN = true;
-			PlaneVertInfo.Color = FVector3f(1, 1, 1);
-			int VertStart[2]{ -1, -1 };
-			for (int MeshIdx = 0; MeshIdx < NumMeshes; MeshIdx++)
-			{
-				PlaneVertInfo.Normal = Normal;
-				if (MeshIdx == 1 && OtherCell != OutsideCellIndex)
-				{
-					PlaneVertInfo.Normal *= -1.0f;
-				}
-				VertStart[MeshIdx] = Meshes[MeshIdx]->MaxVertexID();
-				FVector2f MinUV(FMathf::MaxReal, FMathf::MaxReal);
-				for (int BoundaryVertex : PlaneBoundary)
-				{
-					FVector3d Position = FVector3d(Cells.PlaneBoundaryVertices[BoundaryVertex]);
-					FVector2f UV = FVector2f(PlaneFrame.ToPlaneUV(Position));
-					MinUV.X = FMathf::Min(UV.X, MinUV.X);
-					MinUV.Y = FMathf::Min(UV.Y, MinUV.Y);
-				}
-				for (int BoundaryVertex : PlaneBoundary)
-				{
-					PlaneVertInfo.Position = FVector3d(Cells.PlaneBoundaryVertices[BoundaryVertex]);
-					PlaneVertInfo.UV = (FVector2f(PlaneFrame.ToPlaneUV(PlaneVertInfo.Position)) - MinUV) * GlobalUVScale;
-					Meshes[MeshIdx]->AppendVertex(PlaneVertInfo);
-				}
-			}
-
-			int MID = PlaneToMaterial(PlaneIdx);
-			if (Cells.AssumeConvexCells)
-			{
-				// put a fan
-				for (int V0 = 0, V1 = 1, V2 = 2; V2 < PlaneBoundary.Num(); V1 = V2++)
-				{
-					for (int MeshIdx = 0; MeshIdx < NumMeshes; MeshIdx++)
-					{
-						int Offset = VertStart[MeshIdx];
-						FIndex3i Tri(V0 + Offset, V1 + Offset, V2 + Offset);
-						if (MeshIdx == 1 && OtherCell != OutsideCellIndex)
-						{
-							Swap(Tri.B, Tri.C);
-						}
-						int TID = Meshes[MeshIdx]->AppendTriangle(Tri);
-						if (ensure(TID > -1))
-						{
-							Meshes[MeshIdx]->Attributes()->GetMaterialID()->SetNewValue(TID, MID);
-						}
-					}
-				}
-			}
-			else // cells may not be convex; cannot triangulate w/ fan
-			{
-				// Delaunay triangulate
-				FPolygon2f Polygon;
-				for (int V = 0; V < PlaneBoundary.Num(); V++)
-				{
-					Polygon.AppendVertex(Meshes[0]->GetVertexUV(VertStart[0] + V));
-				}
-
-				FGeneralPolygon2f GeneralPolygon(Polygon);
-				FConstrainedDelaunay2f Triangulation;
-				Triangulation.FillRule = FConstrainedDelaunay2f::EFillRule::NonZero;
-				Triangulation.Add(GeneralPolygon);
-				Triangulation.Triangulate();
-
-				for (int MeshIdx = 0; MeshIdx < NumMeshes; MeshIdx++)
-				{
-					int Offset = VertStart[MeshIdx];
-					for (FIndex3i Triangle : Triangulation.Triangles)
-					{
-						Triangle.A += Offset;
-						Triangle.B += Offset;
-						Triangle.C += Offset;
-						if (MeshIdx == 1 && OtherCell != OutsideCellIndex)
-						{
-							Swap(Triangle.B, Triangle.C);
-						}
-						int TID = Meshes[MeshIdx]->AppendTriangle(Triangle);
-						if (ensure(TID > -1))
-						{
-							Meshes[MeshIdx]->Attributes()->GetMaterialID()->SetNewValue(TID, MID);
-						}
-					}
-				}
-			}
-		}
-	}
-
-	// Approximately calculate a "safe" spacing that would not require the remesher to create more than a million new vertices
-	double GetSafeNoiseSpacing(float SurfaceArea, float TargetSpacing)
-	{
-		double MaxVerts = 1000000;
-		double MinEdgeLen = FMathd::Sqrt((double)SurfaceArea / MaxVerts);
-		double Spacing = FMath::Max3(.001, MinEdgeLen, (double)TargetSpacing);
-		if (Spacing > TargetSpacing)
-		{
-			UE_LOG(LogPlanarCut, Warning,
-				TEXT("Requested spacing of noise points (surface resolution) of %f would require too many added vertices; Using %f instead."),
-				TargetSpacing, Spacing);
-		}
-		return Spacing;
-	}
-
-	void CreateMeshesForBoundedPlanesWithNoise(int NumCells, const FPlanarCells& Cells, const FAxisAlignedBox3d& DomainBounds, bool bNoise, double GlobalUVScale)
-	{
-		TArray<FDynamicMesh3> PlaneMeshes;
-		PlaneMeshes.SetNum(Cells.Planes.Num());
-		FName OriginalPositionAttribute = "OriginalPosition";
-		for (FDynamicMesh3& PlaneMesh : PlaneMeshes)
-		{
-			PlaneMesh.EnableVertexUVs(FVector2f(0, 0));
-			PlaneMesh.EnableVertexNormals(FVector3f::UnitZ());
-			PlaneMesh.EnableVertexColors(FVector3f(1, 1, 1));
-			PlaneMesh.EnableAttributes();
-			PlaneMesh.Attributes()->EnableMaterialID();
-			PlaneMesh.Attributes()->AttachAttribute(OriginalPositionAttribute, new TDynamicMeshVertexAttribute<double, 3>(&PlaneMesh));
-		}
-
-		struct FPlaneIdxAndFlip
-		{
-			int32 PlaneIdx;
-			bool bIsFlipped;
-		};
-		TArray<TArray<FPlaneIdxAndFlip>> CellPlanes; // per cell, the planes that border that cell
-		CellPlanes.SetNum(NumCells);
-
-		for (int32 PlaneIdx = 0; PlaneIdx < Cells.PlaneCells.Num(); PlaneIdx++)
-		{
-			const TPair<int32, int32>& CellPair = Cells.PlaneCells[PlaneIdx];
-			int32 OtherCell = CellPair.Value < 0 ? OutsideCellIndex : CellPair.Value;
-			if (ensure(CellPlanes.IsValidIndex(CellPair.Key)))
-			{
-				CellPlanes[CellPair.Key].Add({ PlaneIdx, false });
-			}
-			if (CellPlanes.IsValidIndex(OtherCell))
-			{
-				CellPlanes[OtherCell].Add({ PlaneIdx, true });
-			}
-		}
-
-		// heuristic to protect against creating too many vertices on remeshing
-		float TotalArea = 0;
-		for (int32 PlaneIdx = 0; PlaneIdx < Cells.Planes.Num(); PlaneIdx++)
-		{
-			const TArray<int>& PlaneBoundary = Cells.PlaneBoundaries[PlaneIdx];
-			const FVector& V0 = Cells.PlaneBoundaryVertices[PlaneBoundary[0]];
-			FVector AreaVec = FVector::ZeroVector;
-			for (int32 V1Idx = 1, V2Idx = 2; V2Idx < PlaneBoundary.Num(); V1Idx = V2Idx++)
-			{
-				const FVector& V1 = Cells.PlaneBoundaryVertices[PlaneBoundary[V1Idx]];
-				const FVector& V2 = Cells.PlaneBoundaryVertices[PlaneBoundary[V2Idx]];
-				AreaVec += (V1 - V0) ^ (V2 - V1);
-			}
-			TotalArea += AreaVec.Size();
-		}
-		double Spacing = GetSafeNoiseSpacing(TotalArea, Cells.InternalSurfaceMaterials.NoiseSettings->PointSpacing);
-		
-		ParallelFor(Cells.Planes.Num(), [this, OriginalPositionAttribute, &PlaneMeshes, &Cells, GlobalUVScale, Spacing](int32 PlaneIdx)
-			{
-				FDynamicMesh3& Mesh = PlaneMeshes[PlaneIdx];
-				const TArray<int>& PlaneBoundary = Cells.PlaneBoundaries[PlaneIdx];
-				FVector3f Normal(Cells.Planes[PlaneIdx].GetNormal());
-				FFrame3d PlaneFrame(Cells.Planes[PlaneIdx]);
-				FVertexInfo PlaneVertInfo;
-				PlaneVertInfo.bHaveC = true;
-				PlaneVertInfo.bHaveUV = true;
-				PlaneVertInfo.bHaveN = true;
-				PlaneVertInfo.Normal = Normal;
-				PlaneVertInfo.UV = FVector2f(0, 0); // UVs will be set below, after noise is added
-				PlaneVertInfo.Color = FVector3f(1, 1, 1);
-
-				FPolygon2f Polygon;
-				for (int BoundaryVertex : PlaneBoundary)
-				{
-					PlaneVertInfo.Position = FVector3d(Cells.PlaneBoundaryVertices[BoundaryVertex]);
-					Polygon.AppendVertex((FVector2f)PlaneFrame.ToPlaneUV(PlaneVertInfo.Position));
-					Mesh.AppendVertex(PlaneVertInfo);
-				}
-
-				// we do a CDT here to give a slightly better start to remeshing; we could try simple ear clipping instead
-				FGeneralPolygon2f GeneralPolygon(Polygon);
-				FConstrainedDelaunay2f Triangulation;
-				Triangulation.FillRule = FConstrainedDelaunay2f::EFillRule::NonZero;
-				Triangulation.Add(GeneralPolygon);
-				Triangulation.Triangulate();
-				if (Triangulation.Triangles.Num() == 0) // fall back to ear clipping if the triangulation came back empty
-				{
-					PolygonTriangulation::TriangulateSimplePolygon(Polygon.GetVertices(), Triangulation.Triangles);
-				}
-				if (ensure(Triangulation.Triangles.Num() > 0))
-				{
-					int MID = PlaneToMaterial(PlaneIdx);
-					for (FIndex3i Triangle : Triangulation.Triangles)
-					{
-						int TID = Mesh.AppendTriangle(Triangle);
-						if (ensure(TID > -1))
-						{
-							Mesh.Attributes()->GetMaterialID()->SetNewValue(TID, MID);
-						}
-					}
-
-					RemeshForNoise(Mesh, EEdgeRefineFlags::SplitsOnly, Spacing);
-					TDynamicMeshVertexAttribute<double, 3>* OriginalPosns =
-						static_cast<TDynamicMeshVertexAttribute<double, 3>*>(Mesh.Attributes()->GetAttachedAttribute(OriginalPositionAttribute));
-					for (int VID : Mesh.VertexIndicesItr())
-					{
-						OriginalPosns->SetValue(VID, Mesh.GetVertex(VID));
-					}
-					ApplyNoise(Mesh, FVector3d(Normal), Cells.InternalSurfaceMaterials.NoiseSettings.GetValue());
-
-					FMeshNormals::QuickComputeVertexNormals(Mesh);
-				}
-			}, EParallelForFlags::None);
-
-		for (int CellIdx = 0; CellIdx < NumCells; CellIdx++)
-		{
-			FCellInfo& CellInfo = CellMeshes[CellIdx];
-			FDynamicMesh3& Mesh = CellInfo.AugMesh;
-			Mesh.Attributes()->AttachAttribute(OriginalPositionAttribute, new TDynamicMeshVertexAttribute<double, 3>(&Mesh));
-			bool bFlipForOutsideCell = CellIdx == OutsideCellIndex; // outside cell will be subtracted, and needs all planes flipped vs normal
-			for (FPlaneIdxAndFlip PlaneInfo : CellPlanes[CellIdx])
-			{
-				AppendMesh(Mesh, PlaneMeshes[PlaneInfo.PlaneIdx], PlaneInfo.bIsFlipped ^ bFlipForOutsideCell);
-			}
-		}
-
-		// resolve self-intersections
-
-		// build hash grid of mesh vertices so we correspond all same-pos vertices across touching meshes
-		TPointHashGrid3d<FIndex2i> MeshesVertices(FMathd::ZeroTolerance*1000, FIndex2i::Invalid());
-		for (int CellIdx = 0; CellIdx < NumCells; CellIdx++)
-		{
-			FCellInfo& CellInfo = CellMeshes[CellIdx];
-			FDynamicMesh3& Mesh = CellInfo.AugMesh;
-			for (int VID : Mesh.VertexIndicesItr())
-			{
-				MeshesVertices.InsertPointUnsafe(FIndex2i(CellIdx, VID), Mesh.GetVertex(VID));
-			}
-		}
-		
-		// repeatedly detect and resolve collisions until there are no more (or give up after too many iterations)
-		TArray<bool> CellUnmoved; CellUnmoved.Init(false, NumCells);
-		const int MaxIters = 10;
-		for (int Iters = 0; Iters < MaxIters; Iters++)
-		{
-			struct FUpdate
-			{
-				FIndex2i Tris;
-				TArray<FIndex2i> IDs;
-				FUpdate(int TriA = -1, int TriB = -1) : Tris(TriA, TriB)
-				{}
-			};
-
-			// todo: can parallelize?
-			TArray<TArray<FUpdate>> Updates; Updates.SetNum(NumCells);
-			bool bAnyUpdatesNeeded = false;
-			for (int CellIdx = 0; CellIdx < NumCells; CellIdx++)
-			{
-				if (CellUnmoved[CellIdx])
-				{
-					// if nothing moved since last time we resolved self intersections on this cell, don't need to process again
-					continue;
-				}
-				FDynamicMesh3& Mesh = CellMeshes[CellIdx].AugMesh;
-				FDynamicMeshAABBTree3 CellTree(&Mesh, true);
-				MeshIntersection::FIntersectionsQueryResult Intersections = CellTree.FindAllSelfIntersections(true);
-				for (MeshIntersection::FSegmentIntersection& Seg : Intersections.Segments)
-				{
-					// manually check for shared edges by vertex position because they might not be topologically connected
-					FIndex3i Tri[2]{ Mesh.GetTriangle(Seg.TriangleID[0]), Mesh.GetTriangle(Seg.TriangleID[1]) };
-					int MatchedVertices = 0;
-					for (int T0SubIdx = 0; T0SubIdx < 3; T0SubIdx++)
-					{
-						FVector3d V0 = Mesh.GetVertex(Tri[0][T0SubIdx]);
-						for (int T1SubIdx = 0; T1SubIdx < 3; T1SubIdx++)
-						{
-							FVector3d V1 = Mesh.GetVertex(Tri[1][T1SubIdx]);
-							if (DistanceSquared(V0, V1) < FMathd::ZeroTolerance)
-							{
-								MatchedVertices++;
-								break;
-							}
-						}
-					}
-					// no shared vertices: treat as a real collision
-					// (TODO: only skip shared edges? will need to do something to avoid shared vertices becoming collisions)
-					if (MatchedVertices < 1)
-					{
-						bAnyUpdatesNeeded = true;
-						FUpdate& Update = Updates[CellIdx].Emplace_GetRef(Seg.TriangleID[0], Seg.TriangleID[1]);
-						for (int TriIdx = 0; TriIdx < 2; TriIdx++)
-						{
-							for (int VSubIdx = 0; VSubIdx < 3; VSubIdx++)
-							{
-								int VIdx = Tri[TriIdx][VSubIdx];
-								FVector3d P = Mesh.GetVertex(VIdx);
-								FIndex2i IDs(CellIdx, VIdx);
-								MeshesVertices.FindPointsInBall(P, FMathd::ZeroTolerance, [this, P](FIndex2i IDs)
-									{
-										FVector3d Pos = CellMeshes[IDs.A].AugMesh.GetVertex(IDs.B);
-										return DistanceSquared(P, Pos);
-									}, Update.IDs);
-							}
-						}
-					}
-				}
-			}
-			if (!bAnyUpdatesNeeded)
-			{
-				break;
-			}
-			for (int CellIdx = 0; CellIdx < NumCells; CellIdx++)
-			{
-				CellUnmoved[CellIdx] = true;
-			}
-			
-			// todo: maybe can parallelize if movements are not applied until after?
-			for (int CellIdx = 0; CellIdx < NumCells; CellIdx++)
-			{
-				FDynamicMesh3& Mesh = CellMeshes[CellIdx].AugMesh;
-				TDynamicMeshVertexAttribute<double, 3>* OriginalPosns =
-					static_cast<TDynamicMeshVertexAttribute<double, 3>*>(Mesh.Attributes()->GetAttachedAttribute(OriginalPositionAttribute));
-				auto InterpVert = [&Mesh, &OriginalPosns](int VID, double t)
-				{
-					FVector3d OrigPos, NoisePos;
-					OriginalPosns->GetValue(VID, OrigPos);
-					NoisePos = Mesh.GetVertex(VID);
-					return UE::Geometry::Lerp(OrigPos, NoisePos, t);
-				};
-				auto InterpTri = [&Mesh, &InterpVert](int TID, double t)
-				{
-					FIndex3i TriVIDs = Mesh.GetTriangle(TID);
-					FTriangle3d Tri;
-					for (int i = 0; i < 3; i++)
-					{
-						Tri.V[i] = InterpVert(TriVIDs[i], t);
-					}
-					return Tri;
-				};
-				auto TestIntersection = [&InterpTri](int TIDA, int TIDB, double t)
-				{
-					FIntrTriangle3Triangle3d TriTri(InterpTri(TIDA, t), InterpTri(TIDB, t));
-					return TriTri.Find();
-				};
-				// resolve tri-tri intersections on this cell's mesh (moving associated verts on other meshes as needed also)
-				for (FUpdate& Update : Updates[CellIdx])
-				{
-					double tsafe = 0;
-					double tbad = 1;
-					if (!TestIntersection(Update.Tris.A, Update.Tris.B, tbad))
-					{
-						continue;
-					}
-					for (int SearchSteps = 0; SearchSteps < 4; SearchSteps++)
-					{
-						double tmid = (tsafe + tbad) * .5;
-						if (TestIntersection(Update.Tris.A, Update.Tris.B, tmid))
-						{
-							tbad = tmid;
-						}
-						else
-						{
-							tsafe = tmid;
-						}
-					}
-					CellUnmoved[CellIdx] = false;
-					for (FIndex2i IDs : Update.IDs)
-					{
-						FVector3d OldPos = CellMeshes[IDs.A].AugMesh.GetVertex(IDs.B);
-						FVector3d NewPos;
-						if (IDs.A == CellIdx)
-						{
-							NewPos = InterpVert(IDs.B, tsafe);
-							Mesh.SetVertex(IDs.B, NewPos);
-						}
-						else
-						{
-							CellUnmoved[IDs.A] = false;
-							FDynamicMesh3& OtherMesh = CellMeshes[IDs.A].AugMesh;
-							TDynamicMeshVertexAttribute<double, 3>* OtherOriginalPosns =
-								static_cast<TDynamicMeshVertexAttribute<double, 3>*>(OtherMesh.Attributes()->GetAttachedAttribute(OriginalPositionAttribute));
-							FVector3d OrigPos;
-							OtherOriginalPosns->GetValue(IDs.B, OrigPos);
-							NewPos = UE::Geometry::Lerp(OrigPos, OldPos, tsafe);
-							OtherMesh.SetVertex(IDs.B, NewPos);
-						}
-						MeshesVertices.UpdatePoint(IDs, OldPos, NewPos);
-					}
-				}
-			}
-		}
-		
-		// clear "original position" attribute now that we have removed self-intersections
-		for (int CellIdx = 0; CellIdx < NumCells; CellIdx++)
-		{
-			FCellInfo& CellInfo = CellMeshes[CellIdx];
-			FDynamicMesh3& Mesh = CellInfo.AugMesh;
-			Mesh.Attributes()->RemoveAttribute(OriginalPositionAttribute);
-		}
-
-		// recompute UVs using new positions after noise was applied + fixed
-		TArray<FVector2f> PlaneMinUVs; PlaneMinUVs.Init(FVector2f(FMathf::MaxReal, FMathf::MaxReal), Cells.Planes.Num());
-		TArray<FFrame3d> PlaneFrames; PlaneFrames.Reserve(Cells.Planes.Num());
-		for (int PlaneIdx = 0; PlaneIdx < Cells.Planes.Num(); PlaneIdx++)
-		{
-			PlaneFrames.Emplace(Cells.Planes[PlaneIdx]);
-		}
-		// first pass to compute min UV for each plane
-		for (FCellInfo& CellInfo : CellMeshes)
-		{
-			FDynamicMesh3& Mesh = CellInfo.AugMesh;
-			FDynamicMeshMaterialAttribute* MaterialIDs = Mesh.Attributes()->GetMaterialID();
-			
-			for (int TID : Mesh.TriangleIndicesItr())
-			{
-				int PlaneIdx = MaterialToPlane(MaterialIDs->GetValue(TID));
-				if (PlaneIdx > -1)
-				{
-					FIndex3i Tri = Mesh.GetTriangle(TID);
-					for (int Idx = 0; Idx < 3; Idx++)
-					{
-						FVector2f UV = FVector2f(PlaneFrames[PlaneIdx].ToPlaneUV(Mesh.GetVertex(Tri[Idx])));
-						FVector2f& MinUV = PlaneMinUVs[PlaneIdx];
-						MinUV.X = FMathf::Min(UV.X, MinUV.X);
-						MinUV.Y = FMathf::Min(UV.Y, MinUV.Y);
-					}
-				}
-			}
-		}
-		// second pass to actually set UVs
-		for (FCellInfo& CellInfo : CellMeshes)
-		{
-			FDynamicMesh3& Mesh = CellInfo.AugMesh;
-			FDynamicMeshMaterialAttribute* MaterialIDs = Mesh.Attributes()->GetMaterialID();
-
-			for (int TID : Mesh.TriangleIndicesItr())
-			{
-				int PlaneIdx = MaterialToPlane(MaterialIDs->GetValue(TID));
-				if (PlaneIdx > -1)
-				{
-					FIndex3i Tri = Mesh.GetTriangle(TID);
-					for (int Idx = 0; Idx < 3; Idx++)
-					{
-						FVector2f UV = ((FVector2f)PlaneFrames[PlaneIdx].ToPlaneUV(Mesh.GetVertex(Tri[Idx])) - PlaneMinUVs[PlaneIdx]) * GlobalUVScale;
-						Mesh.SetVertexUV(Tri[Idx], UV);
-					}
-				}
-			}
-		}
-	}
-	
-	void CreateMeshesForSinglePlane(const FPlanarCells& Cells, const FAxisAlignedBox3d& DomainBounds, bool bNoise, double GlobalUVScale, double Grout, bool bOnlyGrout)
-	{
-		bool bHasGrout = Grout > 0;
-
-		int MID = PlaneToMaterial(0);
-		FPlane Plane = Cells.Planes[0];
-
-		FFrame3d PlaneFrame(Plane);
-		FInterval1d ZRange;
-		FAxisAlignedBox2d XYRange;
-		for (int CornerIdx = 0; CornerIdx < 8; CornerIdx++)
-		{
-			FVector3d Corner = DomainBounds.GetCorner(CornerIdx);
-			XYRange.Contain(PlaneFrame.ToPlaneUV(Corner));
-			ZRange.Contain(Plane.PlaneDot(FVector(Corner)));
-		}
-		//if (FMathd::SignNonZero(ZRange.Min) == FMathd::SignNonZero(ZRange.Max))
-		//{
-		//	// TODO: early out for plane that doesn't even intersect the domain bounding box?
-		//}
-
-		FDynamicMesh3 PlaneMesh(true, true, true, false);
-		FVertexInfo PlaneVertInfo;
-		PlaneVertInfo.bHaveC = true;
-		PlaneVertInfo.bHaveUV = true;
-		PlaneVertInfo.bHaveN = true;
-		PlaneVertInfo.Color = FVector3f(1, 1, 1);
-		PlaneVertInfo.Normal = FVector3f(Plane.GetNormal());
-
-		for (int CornerIdx = 0; CornerIdx < 4; CornerIdx++)
-		{
-			PlaneVertInfo.Position = PlaneFrame.FromPlaneUV(XYRange.GetCorner(CornerIdx));
-			PlaneVertInfo.UV = FVector2f(XYRange.GetCorner(CornerIdx) - XYRange.Min) * GlobalUVScale;
-			PlaneMesh.AppendVertex(PlaneVertInfo);
-		}
-		PlaneMesh.AppendTriangle(0, 1, 2);
-		PlaneMesh.AppendTriangle(0, 2, 3);
-
-		if (bNoise)
-		{
-			double Spacing = GetSafeNoiseSpacing(XYRange.Area(), Cells.InternalSurfaceMaterials.NoiseSettings->PointSpacing);
-			RemeshForNoise(PlaneMesh, EEdgeRefineFlags::SplitsOnly, Spacing);
-			ApplyNoise(PlaneMesh, PlaneFrame.GetAxis(2), Cells.InternalSurfaceMaterials.NoiseSettings.GetValue(), true);
-			FMeshNormals::QuickComputeVertexNormals(PlaneMesh);
-		}
-		TArray<int> PlaneBoundary,  // loop of vertex IDs on the boundary of PlaneMesh (starting with vertex 0)
-			PlaneBoundaryCornerIndices; // indices of the corner vertices in the PlaneBoundary array
-		{
-			double Offset = ZRange.Max;
-			FMeshBoundaryLoops Boundary(&PlaneMesh);
-			checkSlow(Boundary.GetLoopCount() == 1);
-			int FirstIdx;
-			bool bFound = Boundary[0].Vertices.Find(0, FirstIdx);
-			checkSlow(bFound);
-			PlaneBoundary = Boundary[0].Vertices;
-			if (FirstIdx != 0)
-			{
-				Algo::Rotate(PlaneBoundary, FirstIdx);
-			}
-			checkSlow(PlaneBoundary[0] == 0);
-
-			PlaneBoundaryCornerIndices.Add(0);
-			int FoundIndices = 1;
-			for (int VIDIdx = 0; VIDIdx < PlaneBoundary.Num(); VIDIdx++)
-			{
-				int VID = PlaneBoundary[VIDIdx];
-				if (VID == FoundIndices)
-				{
-					FoundIndices++;
-					PlaneBoundaryCornerIndices.Add(VIDIdx);
-				}
-			}
-		}
-		FDynamicMesh3* Meshes[2];
-		if (!bOnlyGrout)
-		{
-			for (int Side = 0; Side < 2; Side++)
-			{
-				Meshes[Side] = &CellMeshes[Side].AugMesh;
-				*Meshes[Side] = PlaneMesh;
-				double Offset = ZRange.Max;
-				TArray<int> CapBoundary, CapBoundaryCornerIndices;
-
-				if (Side == 0)
-				{
-					Meshes[Side]->ReverseOrientation(true);
-					Offset = ZRange.Min;
-				}
-				PlaneVertInfo.Normal = FVector3f(Plane.GetNormal()) * (-1.0f + (float)Side * 2.0f);
-				FVector3d OffsetVec = FVector3d(Plane.GetNormal()) * Offset;
-
-				for (int CornerIdx = 0; CornerIdx < 4; CornerIdx++)
-				{
-					PlaneVertInfo.Position = Meshes[Side]->GetVertex(CornerIdx) + OffsetVec;
-					// UVs shouldn't matter for outer box vertices because they're outside of the domain by construction ...
-					CapBoundary.Add(Meshes[Side]->AppendVertex(PlaneVertInfo));
-					CapBoundaryCornerIndices.Add(CornerIdx);
-				}
-				int NewTris[2]{
-					Meshes[Side]->AppendTriangle(CapBoundary[0], CapBoundary[1], CapBoundary[2]),
-					Meshes[Side]->AppendTriangle(CapBoundary[0], CapBoundary[2], CapBoundary[3])
-				};
-				if (Side == 1)
-				{
-					Meshes[Side]->ReverseTriOrientation(NewTris[0]);
-					Meshes[Side]->ReverseTriOrientation(NewTris[1]);
-				}
-				FDynamicMeshEditor Editor(Meshes[Side]);
-				FDynamicMeshEditResult ResultOut;
-				Editor.StitchSparselyCorrespondedVertexLoops(PlaneBoundary, PlaneBoundaryCornerIndices, CapBoundary, CapBoundaryCornerIndices, ResultOut, Side == 0);
-			}
-		}
-		if (bHasGrout)
-		{
-			int GroutIdx = bOnlyGrout ? 0 : 2;
-			FDynamicMesh3* GroutMesh = &CellMeshes[GroutIdx].AugMesh;
-			FVector3d GroutOffset = (FVector3d)Plane.GetNormal() * (Grout * .5);
-			if (!bOnlyGrout)
-			{
-				for (int Side = 0; Side < 2; Side++)
-				{
-					// shift both sides out by Grout/2
-					MeshTransforms::Translate(*Meshes[Side], GroutOffset * (-1.0 + (double)Side * 2.0));
-				}
-			}
-
-			// make the center (grout) by stitching together two offset copies of PlaneMesh
-			*GroutMesh = PlaneMesh;
-			GroutMesh->ReverseOrientation(true);
-			MeshTransforms::Translate(*GroutMesh, GroutOffset);
-			FMeshIndexMappings IndexMaps;
-			FDynamicMeshEditor Editor(GroutMesh);
-			Editor.AppendMesh(&PlaneMesh, IndexMaps, [GroutOffset](int VID, const FVector3d& PosIn) {return PosIn - GroutOffset;});
-			TArray<int> AppendPlaneBoundary; AppendPlaneBoundary.Reserve(PlaneBoundary.Num());
-			TArray<int> RevBoundary = PlaneBoundary;
-			Algo::Reverse(RevBoundary);
-			for (int VID : RevBoundary)
-			{
-				AppendPlaneBoundary.Add(IndexMaps.GetNewVertex(VID));
-			}
-			FDynamicMeshEditResult ResultOut;
-			Editor.StitchVertexLoopsMinimal(RevBoundary, AppendPlaneBoundary, ResultOut);
-		}
-
-		// fix up custom attributes and material IDs for all meshes
-		for (int CellIdx = 0; CellIdx < CellMeshes.Num(); CellIdx++)
-		{
-			FDynamicMesh3& Mesh = CellMeshes[CellIdx].AugMesh;
-			// re-enable tangents and visibility attributes, since these are lost when we set the mesh to a copy of the plane mesh
-			UE::PlanarCutInternals::AugmentDynamicMesh::Augment(Mesh);
-
-			// Set all material IDs to the one plane's corresponding material ID
-			for (int TID : Mesh.TriangleIndicesItr())
-			{
-				Mesh.Attributes()->GetMaterialID()->SetNewValue(TID, MID);
-			}
-		}
-	}
-};
-
-
-
-// organize metadata corresponding dynamic mesh and geometry collection data
-struct FDynamicMeshCollection
-{
-	struct FMeshData
-	{
-		FDynamicMesh3 AugMesh;
-		
-		// FDynamicMeshAABBTree3 Spatial; // TODO: maybe refactor mesh booleans to allow version where caller provides spatial data; it's computed every boolean now
-		// FTransform3d Transform; // TODO: maybe pretransform the data to a space that is good for cutting; refactor mesh boolean so there is an option to have it not transform input
-		int32 TransformIndex; // where the mesh was from in the geometry collection
-		FTransform ToCollection; // transform that need be applied to go back to the local space of the geometry collection
-
-		FMeshData()
-		{
-			UE::PlanarCutInternals::AugmentDynamicMesh::Augment(AugMesh);
-		}
-
-		FMeshData(const FDynamicMesh3& Mesh, int32 TransformIndex, FTransform ToCollection) : AugMesh(Mesh), TransformIndex(TransformIndex), ToCollection(ToCollection)
-		{}
-	};
-	TIndirectArray<FMeshData> Meshes;
-	FAxisAlignedBox3d Bounds;
-
-	FDynamicMeshCollection(const FGeometryCollection* Collection, const TArrayView<const int32>& TransformIndices, FTransform TransformCollection, bool bSaveIsolatedVertices = false)
-	{
-		Init(Collection, TransformIndices, TransformCollection, bSaveIsolatedVertices);
-	}
-
-	void Init(const FGeometryCollection* Collection, const TArrayView<const int32>& TransformIndices, FTransform TransformCollection, bool bSaveIsolatedVertices = false)
-	{
-		Meshes.Reset();
-		Bounds = FAxisAlignedBox3d::Empty();
-
-		for (int32 TransformIdx : TransformIndices)
-		{
-			if (Collection->Children[TransformIdx].Num() > 0)
-			{
-				// only store the meshes of leaf nodes
-				continue;
-			}
-
-			UE::Geometry::FTransform3d CollectionToLocal = UE::Geometry::FTransform3d(GeometryCollectionAlgo::GlobalMatrix(Collection->Transform, Collection->Parent, TransformIdx) * TransformCollection);
-
-			int32 AddedMeshIdx = Meshes.Add(new FMeshData);
-			FMeshData& MeshData = Meshes[AddedMeshIdx];
-			MeshData.TransformIndex = TransformIdx;
-			MeshData.ToCollection = FTransform(CollectionToLocal.Inverse());
-			FDynamicMesh3& Mesh = MeshData.AugMesh;
-
-			int32 GeometryIdx = Collection->TransformToGeometryIndex[TransformIdx];
-			Mesh.EnableAttributes();
-			Mesh.Attributes()->EnableMaterialID();
-
-			int32 VertexStart = Collection->VertexStart[GeometryIdx];
-			int32 VertexCount = Collection->VertexCount[GeometryIdx];
-			int32 FaceCount = Collection->FaceCount[GeometryIdx];
-
-			FVertexInfo VertexInfo;
-			VertexInfo.bHaveC = true;
-			VertexInfo.bHaveN = true;
-			VertexInfo.bHaveUV = true;
-			for (int32 Idx = VertexStart, N = VertexStart + VertexCount; Idx < N; Idx++)
-			{
-				VertexInfo.Position = CollectionToLocal.TransformPosition(FVector3d(Collection->Vertex[Idx]));
-				VertexInfo.UV = FVector2f(Collection->UV[Idx]);
-				VertexInfo.Color = FVector3f(Collection->Color[Idx]);
-				VertexInfo.Normal = (FVector3f)CollectionToLocal.TransformVectorNoScale(FVector3d(Collection->Normal[Idx]));
-				int VID = Mesh.AppendVertex(VertexInfo);
-				UE::PlanarCutInternals::AugmentDynamicMesh::SetTangent(Mesh, VID, VertexInfo.Normal,
-					(FVector3f)CollectionToLocal.TransformVectorNoScale(FVector3d(Collection->TangentU[Idx])),
-					(FVector3f)CollectionToLocal.TransformVectorNoScale(FVector3d(Collection->TangentV[Idx])));
-			}
-			FIntVector VertexOffset(VertexStart, VertexStart, VertexStart);
-			for (int32 Idx = Collection->FaceStart[GeometryIdx], N = Collection->FaceStart[GeometryIdx] + FaceCount; Idx < N; Idx++)
-			{
-				FIndex3i AddTri = FIndex3i(Collection->Indices[Idx] - VertexOffset);
-				int TID = Mesh.AppendTriangle(AddTri, 0);
-				if (TID == FDynamicMesh3::NonManifoldID)
-				{
-					// work around non-manifold triangles by copying the vertices
-					FIndex3i NewTri(-1, -1, -1);
-					for (int SubIdx = 0; SubIdx < 3; SubIdx++)
-					{
-						int NewVID = Mesh.AppendVertex(Mesh, AddTri[SubIdx]);
-						int32 SrcIdx = AddTri[SubIdx] + VertexStart;
-						UE::PlanarCutInternals::AugmentDynamicMesh::SetTangent(Mesh, NewVID,
-							Mesh.GetVertexNormal(NewVID), // TODO: we don't actually use the vertex normal; consider removing this arg from the function entirely
-							(FVector3f)CollectionToLocal.TransformVectorNoScale(FVector3d(Collection->TangentU[SrcIdx])),
-							(FVector3f)CollectionToLocal.TransformVectorNoScale(FVector3d(Collection->TangentV[SrcIdx])));
-						NewTri[SubIdx] = NewVID;
-					}
-					TID = Mesh.AppendTriangle(NewTri, 0);
-				}
-				if (TID < 0)
-				{
-					continue;
-				}
-				Mesh.Attributes()->GetMaterialID()->SetValue(TID, Collection->MaterialID[Idx]);
-				UE::PlanarCutInternals::AugmentDynamicMesh::SetVisibility(Mesh, TID, Collection->Visible[Idx]);
-				// note: material index doesn't need to be passed through; will be rebuilt by a call to reindex materials once the cut mesh is returned back to geometry collection format
-			}
-
-			if (!bSaveIsolatedVertices)
-			{
-				FDynamicMeshEditor Editor(&Mesh);
-				Editor.RemoveIsolatedVertices();
-			}
-
-			Bounds.Contain(Mesh.GetCachedBounds());
-
-			// TODO: build spatial data (add this after setting up mesh boolean path that can use it)
-			//MeshData.Spatial.SetMesh(&Mesh);
-		}
-	}
-
-
-	int32 CutWithMultiplePlanes(
-		const TArrayView<const FPlane>& Planes, 
-		double Grout,
-		double CollisionSampleSpacing,
-		FGeometryCollection* Collection,
-		FInternalSurfaceMaterials& InternalSurfaceMaterials,
-		bool bSetDefaultInternalMaterialsFromCollection
-		)
-	{
-#if WITH_EDITOR
-		// Create progress indicator dialog
-		static const FText SlowTaskText = NSLOCTEXT("CutMultipleWithMultiplePlanes", "CutMultipleWithMultiplePlanesText", "Cutting geometry collection with plane(s)...");
-
-		FScopedSlowTask SlowTask(Planes.Num(), SlowTaskText);
-		SlowTask.MakeDialog();
-
-		// Declare progress shortcut lambdas
-		auto EnterProgressFrame = [&SlowTask](float Progress)
-		{
-			SlowTask.EnterProgressFrame(Progress);
-		};
-#else
-		auto EnterProgressFrame = [](float Progress) {};
-#endif
-
-		bool bHasGrout = Grout > 0;
-
-		if (bHasGrout)
-		{
-			// For multi-plane cuts with grout specifically, the easiest path seems to be:
-			// 1. Build the "grout" section of each plane
-			// 2. Take the union of all those grout sections as the grout mesh
-			// 3. Use the generic CutWithCellMeshes path, where that grout mesh is both the inner and outside cell mesh
-			//    (Note the outside cell mesh is subtracted, not intersected)
-			//    (Note this relies on island splitting to separate all the pieces afterwards.)
-			FCellMeshes GroutCells;
-			GroutCells.SetNumCells(2);
-			FDynamicMesh3& GroutMesh = GroutCells.CellMeshes[0].AugMesh;
-			FDynamicMeshEditor GroutAppender(&GroutMesh);
-			FMeshIndexMappings IndexMaps;
-			for (int32 PlaneIdx = 0; PlaneIdx < Planes.Num(); PlaneIdx++)
-			{
-				EnterProgressFrame(.5);
-				FPlanarCells PlaneCells(Planes[PlaneIdx]);
-				PlaneCells.InternalSurfaceMaterials = InternalSurfaceMaterials;
-				FCellMeshes PlaneGroutMesh;
-				PlaneGroutMesh.MakeOnlyPlanarGroutCell(PlaneCells, Bounds, Grout);
-				GroutAppender.AppendMesh(&PlaneGroutMesh.CellMeshes[0].AugMesh, IndexMaps);
-			}
-
-			EnterProgressFrame(Planes.Num() * .2);
-			FMeshSelfUnion GroutUnion(&GroutMesh);
-			GroutUnion.bSimplifyAlongNewEdges = true;
-			GroutUnion.bWeldSharedEdges = false;
-			GroutUnion.Compute();
-
-			EnterProgressFrame(Planes.Num() * .1);
-			// first mesh is the same as the second mesh, but will be subtracted b/c it's the "outside cell"
-			GroutCells.CellMeshes[1].AugMesh = GroutMesh;
-			GroutCells.OutsideCellIndex = 1;
-
-			EnterProgressFrame(Planes.Num() * .2);
-			TArray<TPair<int32, int32>> CellConnectivity;
-			CellConnectivity.Add(TPair<int32, int32>(0, -1));
-
-			return CutWithCellMeshes(InternalSurfaceMaterials, CellConnectivity, GroutCells, Collection, bSetDefaultInternalMaterialsFromCollection, CollisionSampleSpacing);
-		}
-
-		bool bHasProximity = Collection->HasAttribute("Proximity", FGeometryCollection::GeometryGroup);
-		TArray<TUniquePtr<FMeshData>> ToCut;
-		TArray<TUniquePtr<TPointHashGrid3d<int>>> VerticesHashes;
-		auto HashMeshVertices = [&VerticesHashes, &ToCut](int32 HashIdx)
-		{
-			FDynamicMesh3& Mesh = ToCut[HashIdx]->AugMesh;
-			if (HashIdx >= VerticesHashes.Num())
-			{
-				VerticesHashes.SetNum(HashIdx + 1);
-			}
-			if (VerticesHashes[HashIdx].IsValid())
-			{
-				return;
-			}
-			VerticesHashes[HashIdx] = MakeUnique<TPointHashGrid3d<int>>(FMathd::ZeroTolerance * 1000, -1);
-			TPointHashGrid3d<int>& Grid = *VerticesHashes[HashIdx].Get();
-			for (int VID : Mesh.VertexIndicesItr())
-			{
-				Grid.InsertPointUnsafe(VID, Mesh.GetVertex(VID));
-			}
-		};
-		auto ClearHash = [&VerticesHashes](int32 HashIdx)
-		{
-			if (HashIdx < VerticesHashes.Num())
-			{
-				VerticesHashes[HashIdx].Release();
-			}
-		};
-		auto IsNeighbor = [&ToCut, &VerticesHashes](int32 A, int32 B)
-		{
-			if (!ensure(A < ToCut.Num() && B < ToCut.Num() && A < VerticesHashes.Num() && B < VerticesHashes.Num()))
-			{
-				return false;
-			}
-			if (!ensure(VerticesHashes[A].IsValid() && VerticesHashes[B].IsValid()))
-			{
-				return false;
-			}
-			if (!ToCut[A]->AugMesh.GetCachedBounds().Intersects(ToCut[B]->AugMesh.GetCachedBounds()))
-			{
-				return false;
-			}
-			if (ToCut[A]->AugMesh.VertexCount() > ToCut[B]->AugMesh.VertexCount())
-			{
-				Swap(A, B);
-			}
-			FDynamicMesh3& RefMesh = ToCut[B]->AugMesh;
-			for (const FVector3d& V : ToCut[A]->AugMesh.VerticesItr())
-			{
-				TPair<int, double> Nearest = VerticesHashes[B]->FindNearestInRadius(V, FMathd::ZeroTolerance * 10, [&RefMesh, &V](int VID)
-					{
-						return DistanceSquared(RefMesh.GetVertex(VID), V);
-					});
-				if (Nearest.Key != -1)
-				{
-					return true;
-				}
-			}
-			return false;
-		};
-		// copy initial surfaces
-		for (FMeshData& MeshData : Meshes)
-		{
-			ToCut.Add(MakeUnique<FMeshData>(MeshData));
-		}
-		// track connections between meshes via their indices in the ToCut array
-		TMultiMap<int32, int32> Proximity;
-		auto ProxLink = [&Proximity](int32 A, int32 B)
-		{
-			Proximity.Add(A, B);
-			Proximity.Add(B, A);
-		};
-		auto ProxUnlink = [&Proximity](int32 A, int32 B)
-		{
-			Proximity.RemoveSingle(A, B);
-			Proximity.RemoveSingle(B, A);
-		};
-		for (int32 PlaneIdx = 0; PlaneIdx < Planes.Num(); PlaneIdx++)
-		{
-			EnterProgressFrame(1);
-			FPlanarCells PlaneCells(Planes[PlaneIdx]);
-			PlaneCells.InternalSurfaceMaterials = InternalSurfaceMaterials;
-			double OnePercentExtend = Bounds.MaxDim() * .01;
-			FCellMeshes CellMeshes(PlaneCells, Bounds, 0, OnePercentExtend, false);
-
-			// TODO: we could do these cuts in parallel (will takes some rework of the proximity and how results are added to the ToCut array)
-			for (int32 ToCutIdx = 0, ToCutNum = ToCut.Num(); ToCutIdx < ToCutNum; ToCutIdx++)
-			{
-				FMeshData& Surface = *ToCut[ToCutIdx];
-				int32 TransformIndex = Surface.TransformIndex;
-				FTransform ToCollection = Surface.ToCollection;
-
-				FAxisAlignedBox3d Box = Surface.AugMesh.GetCachedBounds();
-				if (InternalSurfaceMaterials.NoiseSettings.IsSet())
-				{
-					Box.Expand(InternalSurfaceMaterials.NoiseSettings->Amplitude);
-				}
-				if (!FMath::PlaneAABBIntersection(Planes[PlaneIdx], FBox(Box)))
-				{
-					continue;
-				}
-
-				TArray<TUniquePtr<FMeshData>> BoolResults;
-				for (int ResultIdx = 0; ResultIdx < 2; ResultIdx++)
-				{
-					BoolResults.Add(MakeUnique<FMeshData>());
-					BoolResults[ResultIdx]->TransformIndex = TransformIndex;
-					BoolResults[ResultIdx]->ToCollection = ToCollection;
-				}
-				check(CellMeshes.CellMeshes.Num() == 2);
-				bool bKeepResults = true;
-				for (int32 CellIdx = 0; CellIdx < 2; CellIdx++)
-				{
-					FCellMeshes::FCellInfo& Cell = CellMeshes.CellMeshes[CellIdx];
-
-					FMeshBoolean::EBooleanOp Op = FMeshBoolean::EBooleanOp::Intersect;
-					if (CellIdx == CellMeshes.OutsideCellIndex)
-					{
-						Op = FMeshBoolean::EBooleanOp::Difference;
-					}
-					FMeshBoolean Boolean(&Surface.AugMesh, &Cell.AugMesh, &BoolResults[CellIdx]->AugMesh, Op);
-					Boolean.bSimplifyAlongNewEdges = true;
-					Boolean.PreserveUVsOnlyForMesh = 0; // slight warping of the autogenerated cell UVs generally doesn't matter
-					Boolean.bWeldSharedEdges = false;
-					if (!Boolean.Compute())
-					{
-						// TODO: do something about failure cases?  e.g. try auto-filling small holes?
-						// note: failure cases won't be detected at all unless we weld edges,
-						//       which will require re-working how tangents are carried through
-					}
-					if (BoolResults[CellIdx]->AugMesh.TriangleCount() == 0)
-					{
-						bKeepResults = false;
-						break;
-					}
-				}
-
-				if (bKeepResults)
-				{
-					ToCut[ToCutIdx] = MoveTemp(BoolResults[0]);
-					int32 NewIdx = ToCut.Add(MoveTemp(BoolResults[1]));
-					// indices of all boolean result meshes (may be more than two due to splitting disconnected components)
-					TArray<int32, TInlineAllocator<4>> ResultIndices = { ToCutIdx, NewIdx };
-					// corresponding parent indices for each result mesh
-					TArray<int32, TInlineAllocator<4>> ParentIndices = { 0, 1 };
-					TArray<FDynamicMesh3> SplitMeshes;
-					for (int UnsplitIdx = 0; UnsplitIdx < 2; UnsplitIdx++)
-					{
-						if (SplitIslands(ToCut[ResultIndices[UnsplitIdx]]->AugMesh, SplitMeshes))
-						{
-							ToCut[ResultIndices[UnsplitIdx]]->AugMesh = SplitMeshes[0];
-							for (int32 Idx = 1; Idx < SplitMeshes.Num(); Idx++)
-							{
-								const FDynamicMesh3& Mesh = SplitMeshes[Idx];
-								ResultIndices.Add(ToCut.Add(MakeUnique<FMeshData>(Mesh, TransformIndex, ToCollection)));
-								ParentIndices.Add(UnsplitIdx);
-							}
-						}
-					}
-
-					// update proximity for neighbors of the original piece
-					if (bHasProximity)
-					{
-						ClearHash(ToCutIdx);
-						TArray<int32> Nbrs;
-						Proximity.MultiFind(ToCutIdx, Nbrs);
-						if (Nbrs.Num() > 0)
-						{
-							for (int32 ChangedMeshIdx : ResultIndices)
-							{
-								HashMeshVertices(ChangedMeshIdx);
-							}
-
-							for (int32 Nbr : Nbrs)
-							{
-								ProxUnlink(ToCutIdx, Nbr);
-								HashMeshVertices(Nbr);
-								for (int32 Idx = 0; Idx < ResultIndices.Num(); Idx++)
-								{
-									int32 ResultIdx = ResultIndices[Idx];
-									int32 OldIdx = Nbr;
-									if (IsNeighbor(ResultIdx, OldIdx))
-									{
-										ProxLink(ResultIdx, OldIdx);
-									}
-								}
-							}
-						}
-
-						if (ResultIndices.Num() == 2)
-						{
-							// add the connection between the two new pieces
-							ProxLink(ResultIndices[0], ResultIndices[1]);
-						}
-						else
-						{
-							if (Nbrs.Num() == 0)
-							{
-								for (int32 ChangedMeshIdx : ResultIndices)
-								{
-									HashMeshVertices(ChangedMeshIdx);
-								}
-							}
-							// check for connections between all pieces
-							for (int FirstIdx = 0; FirstIdx + 1 < ResultIndices.Num(); FirstIdx++)
-							{
-								int32 FirstParent = ParentIndices[FirstIdx];
-								for (int SecondIdx = FirstIdx + 1; SecondIdx < ResultIndices.Num(); SecondIdx++)
-								{
-									if (FirstParent == ParentIndices[SecondIdx])
-									{
-										// these pieces split from the same mesh *because* they were disconnected, so the pieces cannot be neighbors
-										continue;
-									}
-									if (IsNeighbor(ResultIndices[FirstIdx], ResultIndices[SecondIdx]))
-									{
-										ProxLink(ResultIndices[FirstIdx], ResultIndices[SecondIdx]);
-									}
-								}
-							}
-						}
-					}
-				} // iteration over meshes to cut
-			} // iteration over cutting planes
-		}
-
-		TMultiMap<int32, int32> ParentTransformToChildren;
-		for (int32 ToCutIdx = 0; ToCutIdx < ToCut.Num(); ToCutIdx++)
-		{
-			ParentTransformToChildren.Add(ToCut[ToCutIdx]->TransformIndex, ToCutIdx);
-		}
-
-		TArray<int32> ToCutIdxToGeometryIdx;  ToCutIdxToGeometryIdx.Init(-1, ToCut.Num());
-		TArray<int32> ToCutIndices;
-		int32 FirstCreatedIndex = -1;
-		for (FMeshData& MeshData : Meshes)
-		{
-			int32 GeometryIdx = Collection->TransformToGeometryIndex[MeshData.TransformIndex];
-			int32 InternalMaterialID = bSetDefaultInternalMaterialsFromCollection ? InternalSurfaceMaterials.GetDefaultMaterialIDForGeometry(*Collection, GeometryIdx) : InternalSurfaceMaterials.GlobalMaterialID;
-			ToCutIndices.Reset();
-			ParentTransformToChildren.MultiFind(MeshData.TransformIndex, ToCutIndices);
-
-			// if there's only one mesh here, i.e. it didn't get cut at all
-			if (ToCutIndices.Num() <= 1)
-			{
-				continue;
-			}
-
-			// hide old parent geometry
-			SetVisibility(*Collection, GeometryIdx, false);
-
-			// add newly created geometry as children
-			int32 SubPartIdx = 0;
-			for (int32 ToCutIdx : ToCutIndices)
-			{
-				FDynamicMesh3& Mesh = ToCut[ToCutIdx]->AugMesh;
-
-				FString BoneName = GetBoneName(*Collection, ToCut[ToCutIdx]->TransformIndex, SubPartIdx++);
-				int32 CreatedGeometryIdx = AppendToCollection(ToCut[ToCutIdx]->ToCollection, Mesh, CollisionSampleSpacing, ToCut[ToCutIdx]->TransformIndex, BoneName, *Collection, InternalMaterialID);
-				ToCutIdxToGeometryIdx[ToCutIdx] = CreatedGeometryIdx;
-				if (FirstCreatedIndex == -1)
-				{
-					FirstCreatedIndex = CreatedGeometryIdx;
-				}
-			}
-		}
-
-		// create proximity sets on geometry collection and populate using ToCut's Proximity multimap and the array ToCutIdxToGeometryIdx
-		if (bHasProximity)
-		{
-			TManagedArray<TSet<int32>>& GCProximity = Collection->GetAttribute<TSet<int32>>("Proximity", FGeometryCollection::GeometryGroup);
-			for (TPair<int32, int32> Link : Proximity)
-			{
-				GCProximity[ToCutIdxToGeometryIdx[Link.Key]].Add(ToCutIdxToGeometryIdx[Link.Value]);
-			}
-		}
-
-		return FirstCreatedIndex;
-	}
-
-	/**
-	 * Cut collection meshes with cell meshes, and append results to a geometry collection
-	 *
-	 * @param InternalSurfaceMaterials Internal material info (used for material ID)
-	 * @param CellConnectivity The connectivity between cells: PlaneTag -> The two cells separated by triangles with this tag
-	 * @param CellsMeshes Meshed versions of the cells, with noise and material properties baked in
-	 * @param Collection Results will be stored in this
-	 * @param bSetDefaultInternalMaterialsFromCollection If true, set internal materials to the most common external material + 1, following a convenient artist convention
-	 * @return Index of the first created geometry
-	 */
-	int32 CutWithCellMeshes(const FInternalSurfaceMaterials& InternalSurfaceMaterials, const TArray<TPair<int32, int32>>& CellConnectivity, FCellMeshes& CellMeshes, FGeometryCollection* Collection, bool bSetDefaultInternalMaterialsFromCollection, double CollisionSampleSpacing)
-	{
-		// TODO: should we do these cuts in parallel, and the appends sequentially below?
-		int32 FirstIdx = -1;
-		int BadCount = 0;
-		bool bHasProximity = Collection->HasAttribute("Proximity", FGeometryCollection::GeometryGroup);
-		for (FMeshData& Surface : Meshes)
-		{
-			int32 GeometryIdx = Collection->TransformToGeometryIndex[Surface.TransformIndex];
-			TArray<TUniquePtr<FDynamicMesh3>> BooleanResults;  BooleanResults.SetNum(CellMeshes.CellMeshes.Num());
-			ParallelFor(CellMeshes.CellMeshes.Num(), [&BooleanResults, &CellMeshes, &Surface](int32 CellIdx)
-				{
-					FCellMeshes::FCellInfo& Cell = CellMeshes.CellMeshes[CellIdx];
-					if (Cell.AugMesh.GetCachedBounds().Intersects(Surface.AugMesh.GetCachedBounds()))
-					{
-						BooleanResults[CellIdx] = MakeUnique<FDynamicMesh3>();
-						FDynamicMesh3& AugBoolResult = *BooleanResults[CellIdx];
-
-						FMeshBoolean::EBooleanOp Op = FMeshBoolean::EBooleanOp::Intersect;
-						if (CellIdx == CellMeshes.OutsideCellIndex)
-						{
-							Op = FMeshBoolean::EBooleanOp::Difference;
-						}
-						FMeshBoolean Boolean(&Surface.AugMesh, &Cell.AugMesh, &AugBoolResult, Op);
-						Boolean.bSimplifyAlongNewEdges = true;
-						Boolean.PreserveUVsOnlyForMesh = 0; // slight warping of the autogenerated cell UVs generally doesn't matter
-						Boolean.bWeldSharedEdges = false;
-						if (!Boolean.Compute())
-						{
-							// TODO: do something about failure cases?  e.g. try auto-filling small holes?
-							// note: failure cases won't be detected at all unless we weld edges,
-							//       which will require re-working how tangents are carried through
-						}
-					}
-				}, EParallelForFlags::None);
-
-			int32 NonEmptyResults = 0;
-			for (const TUniquePtr<FDynamicMesh3>& AugBoolResult : BooleanResults)
-			{
-				if (AugBoolResult.IsValid() && AugBoolResult->TriangleCount() > 0)
-				{
-					NonEmptyResults++;
-				}
-			}
-
-			if (NonEmptyResults > 1) // only write to geometry collection if more than one result was non-empty
-			{
-				TSet<int32> PlanesInOutput;
-				TMultiMap<int32, int32> CellToGeometry;
-				TMap<int32, int32> GeometryToResultMesh;
-				int32 SubPartIndex = 0;
-				int32 InternalMaterialID = bSetDefaultInternalMaterialsFromCollection ? InternalSurfaceMaterials.GetDefaultMaterialIDForGeometry(*Collection, GeometryIdx) : InternalSurfaceMaterials.GlobalMaterialID;
-
-				for (int32 CellIdx = 0; CellIdx < CellMeshes.CellMeshes.Num(); CellIdx++)
-				{					
-					if (BooleanResults[CellIdx].IsValid() && BooleanResults[CellIdx]->TriangleCount() > 0)
-					{
-						FDynamicMesh3& AugBoolResult = *BooleanResults[CellIdx];
-						for (int TID : AugBoolResult.TriangleIndicesItr())
-						{
-							int MID = AugBoolResult.Attributes()->GetMaterialID()->GetValue(TID);
-							int32 PlaneIdx = CellMeshes.MaterialToPlane(MID);
-							if (PlaneIdx >= 0)
-							{
-								PlanesInOutput.Add(PlaneIdx);
-							}
-						}
-						int32 CreatedGeometryIdx = -1;
-						TArray<FDynamicMesh3> Islands;
-						if (SplitIslands(AugBoolResult, Islands))
-						{
-							for (int32 i = 0; i < Islands.Num(); i++)
-							{
-								FDynamicMesh3& Island = Islands[i];
-								FString BoneName = GetBoneName(*Collection, Surface.TransformIndex, SubPartIndex++);
-								CreatedGeometryIdx = AppendToCollection(Surface.ToCollection, Island, CollisionSampleSpacing, Surface.TransformIndex, BoneName, *Collection, InternalMaterialID);
-								CellToGeometry.Add(CellIdx, CreatedGeometryIdx);
-								if (i > 0)
-								{
-									GeometryToResultMesh.Add(CreatedGeometryIdx, BooleanResults.Add(MakeUnique<FDynamicMesh3>(Island)));
-								}
-								else
-								{
-									*BooleanResults[CellIdx] = Island;
-									GeometryToResultMesh.Add(CreatedGeometryIdx, CellIdx);
-								}
-							}
-						}
-						else
-						{
-							FString BoneName = GetBoneName(*Collection, Surface.TransformIndex, SubPartIndex++);
-							CreatedGeometryIdx = AppendToCollection(Surface.ToCollection, AugBoolResult, CollisionSampleSpacing, Surface.TransformIndex, BoneName, *Collection, InternalMaterialID);
-							CellToGeometry.Add(CellIdx, CreatedGeometryIdx);
-							GeometryToResultMesh.Add(CreatedGeometryIdx, CellIdx);
-						}
-						if (FirstIdx == -1)
-						{
-							FirstIdx = CreatedGeometryIdx;
-						}
-					}
-				}
-				if (bHasProximity)
-				{
-					TManagedArray<TSet<int32>>& Proximity = Collection->GetAttribute<TSet<int32>>("Proximity", FGeometryCollection::GeometryGroup);
-					TArray<TUniquePtr<TPointHashGrid3d<int>>> VertexHashes;
-					auto MakeHash = [this, &VertexHashes, &BooleanResults](int GID)
-					{
-						if (GID >= VertexHashes.Num())
-						{
-							VertexHashes.SetNum(GID + 1);
-						}
-						if (!VertexHashes[GID].IsValid())
-						{
-							VertexHashes[GID] = MakeUnique<TPointHashGrid3d<int>>(FMathd::ZeroTolerance * 1000, -1);
-							FillVertexHash(*BooleanResults[GID], *VertexHashes[GID]);
-						}
-					};
-					for (int32 PlaneIdx : PlanesInOutput)
-					{
-						TPair<int32, int32> Cells = CellConnectivity[PlaneIdx];
-						int32 SecondCell = Cells.Value < 0 ? CellMeshes.OutsideCellIndex : Cells.Value;
-						if (SecondCell != -1)
-						{
-							TArray<int32, TInlineAllocator<4>> GeomA, GeomB;
-							CellToGeometry.MultiFind(Cells.Key, GeomA, false);
-							CellToGeometry.MultiFind(SecondCell, GeomB, false);
-							if (GeomA.Num() == 1 && GeomB.Num() == 1)
-							{
-								Proximity[GeomA[0]].Add(GeomB[0]);
-								Proximity[GeomB[0]].Add(GeomA[0]);
-							}
-							else if (GeomA.Num() >= 1 && GeomB.Num() >= 1) // at least one was split; need to re-check proximities
-							{
-								for (int GIDA : GeomA)
-								{
-									int MeshA = GeometryToResultMesh[GIDA];
-									MakeHash(MeshA);
-									for (int GIDB : GeomB)
-									{
-										int MeshB = GeometryToResultMesh[GIDB];
-										MakeHash(MeshB);
-										if (IsNeighboring(*BooleanResults[MeshA], *VertexHashes[MeshA], *BooleanResults[MeshB], *VertexHashes[MeshB]))
-										{
-											Proximity[GIDA].Add(GIDB);
-											Proximity[GIDB].Add(GIDA);
-										}
-									}
-								}
-							}
-						}
-					}
-				}
-				// turn off old geom visibility (preferred default behavior)
-				SetVisibility(*Collection, GeometryIdx, false);
-			}
-		}
-
-		return FirstIdx;
-	}
-
-	static void SetVisibility(FGeometryCollection& Collection, int32 GeometryIdx, bool bVisible)
-	{
-		int32 FaceEnd = Collection.FaceCount[GeometryIdx] + Collection.FaceStart[GeometryIdx];
-		for (int32 FaceIdx = Collection.FaceStart[GeometryIdx]; FaceIdx < FaceEnd; FaceIdx++)
-		{
-			Collection.Visible[FaceIdx] = bVisible;
-		}
-	}
-
-	void FillVertexHash(const FDynamicMesh3& Mesh, TPointHashGrid3d<int>& VertHash)
-	{
-		for (int VID : Mesh.VertexIndicesItr())
-		{
-			FVector3d V = Mesh.GetVertex(VID);
-			VertHash.InsertPointUnsafe(VID, V);
-		}
-	}
-
-	bool IsNeighboring(FDynamicMesh3& MeshA, const TPointHashGrid3d<int>& VertHashA, FDynamicMesh3& MeshB, const TPointHashGrid3d<int>& VertHashB)
-	{
-		FDynamicMesh3* Mesh[2]{ &MeshA, &MeshB };
-		const TPointHashGrid3d<int>* VertHash[2]{ &VertHashA, &VertHashB };
-		return IsNeighboring(Mesh, VertHash);
-	}
-
-	bool IsNeighboring(FDynamicMesh3* Mesh[2], const TPointHashGrid3d<int>* VertHash[2])
-	{
-		if (!ensure(Mesh[0] && Mesh[1] && VertHash[0] && VertHash[1]))
-		{
-			return false;
-		}
-		if (!Mesh[0]->GetCachedBounds().Intersects(Mesh[1]->GetCachedBounds()))
-		{
-			return false;
-		}
-		
-		int A = 0, B = 1;
-		if (Mesh[0]->VertexCount() > Mesh[1]->VertexCount())
-		{
-			Swap(A, B);
-		}
-		FDynamicMesh3& RefMesh = *Mesh[B];
-		for (const FVector3d& V : Mesh[A]->VerticesItr())
-		{
-			TPair<int, double> Nearest = VertHash[B]->FindNearestInRadius(V, FMathd::ZeroTolerance * 10, [&RefMesh, &V](int VID)
-				{
-					return DistanceSquared(RefMesh.GetVertex(VID), V);
-				});
-			if (Nearest.Key != -1)
-			{
-				return true;
-			}
-		}
-		return false;
-	}
-
-	// Split mesh into connected components, including implicit connections by co-located vertices
-	bool SplitIslands(FDynamicMesh3& Source, TArray<FDynamicMesh3>& SeparatedMeshes)
-	{
-		double SnapDistance = 1e-03;
-		TPointHashGrid3d<int> VertHash(SnapDistance * 10, -1);
-		FDisjointSet VertComponents(Source.MaxVertexID());
-		// Add Source vertices to hash & disjoint sets
-		TArray<int> Neighbors;
-		for (int VID : Source.VertexIndicesItr())
-		{
-			FVector3d Pt = Source.GetVertex(VID);
-			Neighbors.Reset();
-			VertHash.FindPointsInBall(Pt, SnapDistance, [&Source, Pt](int OtherVID) {return DistanceSquared(Pt, Source.GetVertex(OtherVID));}, Neighbors);
-			for (int NbrVID : Neighbors)
-			{
-				VertComponents.UnionSequential(VID, NbrVID);
-			}
-			VertHash.InsertPointUnsafe(VID, Pt);
-		}
-		for (FIndex3i Tri : Source.TrianglesItr())
-		{
-			VertComponents.Union(Tri.A, Tri.B);
-			VertComponents.Union(Tri.B, Tri.C);
-			VertComponents.Union(Tri.C, Tri.A);
-		}
-		
-		bool bWasSplit = FDynamicMeshEditor::SplitMesh(&Source, SeparatedMeshes, [&Source, &VertComponents](int TID)
-			{
-				return (int)VertComponents.Find(Source.GetTriangle(TID).A);
-			});
-
-		if (bWasSplit)
-		{
-			// disconnected components that are contained inside other components need to be re-merged
-			TMeshSpatialSort<FDynamicMesh3> SpatialSort(SeparatedMeshes);
-			SpatialSort.NestingMethod = TMeshSpatialSort<FDynamicMesh3>::ENestingMethod::InLargestParent;
-			SpatialSort.bOnlyNestNegativeVolumes = false;
-			SpatialSort.bOnlyParentPostiveVolumes = true;
-			SpatialSort.Compute();
-			TArray<bool> KeepMeshes; KeepMeshes.Init(true, SeparatedMeshes.Num());
-			for (TMeshSpatialSort<FDynamicMesh3>::FMeshNesting& Nest : SpatialSort.Nests)
-			{
-				FDynamicMeshEditor Editor(&SeparatedMeshes[Nest.OuterIndex]);
-				FMeshIndexMappings Mappings;
-				for (int Inner : Nest.InnerIndices)
-				{
-					Editor.AppendMesh(&SeparatedMeshes[Inner], Mappings);
-					KeepMeshes[Inner] = false;
-				}
-			}
-			for (int Idx = 0; Idx < SeparatedMeshes.Num(); Idx++)
-			{
-				if (!KeepMeshes[Idx])
-				{
-					SeparatedMeshes.RemoveAtSwap(Idx, 1, false);
-					KeepMeshes.RemoveAtSwap(Idx, 1, false);
-					Idx--;
-				}
-			}
-		}
-		return bWasSplit;
-	}
-
-	FString GetBoneName(FGeometryCollection& Output, int TransformParent, int SubPartIndex)
-	{
-		return Output.BoneName[TransformParent] + "_" + FString::FromInt(SubPartIndex);
-	}
-
-	void AddCollisionSamples(double CollisionSampleSpacing)
-	{
-		for (int32 MeshIdx = 0; MeshIdx < Meshes.Num(); MeshIdx++)
-		{
-			UE::PlanarCutInternals::AugmentDynamicMesh::AddCollisionSamplesPerComponent(Meshes[MeshIdx].AugMesh, CollisionSampleSpacing);
-		}
-	}
-
-	// Update all geometry in a GeometryCollection w/ the meshes in the MeshCollection
-	// Resizes the GeometryCollection as needed
-	bool UpdateAllCollections(FGeometryCollection& Collection)
-	{
-		bool bAllSucceeded = true;
-
-		int32 NumGeometry = Collection.NumElements(FGeometryCollection::GeometryGroup);
-		TArray<int32> NewFaceCounts, NewVertexCounts;
-		NewFaceCounts.SetNumUninitialized(NumGeometry);
-		NewVertexCounts.SetNumUninitialized(NumGeometry);
-		for (int32 GeomIdx = 0; GeomIdx < Collection.FaceCount.Num(); GeomIdx++)
-		{
-			NewFaceCounts[GeomIdx] = Collection.FaceCount[GeomIdx];
-			NewVertexCounts[GeomIdx] = Collection.VertexCount[GeomIdx];
-		}
-		for (int MeshIdx = 0; MeshIdx < Meshes.Num(); MeshIdx++)
-		{
-			FDynamicMeshCollection::FMeshData& MeshData = Meshes[MeshIdx];
-			int32 GeomIdx = Collection.TransformToGeometryIndex[MeshData.TransformIndex];
-			NewFaceCounts[GeomIdx] = MeshData.AugMesh.TriangleCount();
-			NewVertexCounts[GeomIdx] = MeshData.AugMesh.VertexCount();
-		}
-		GeometryCollectionAlgo::ResizeGeometries(&Collection, NewFaceCounts, NewVertexCounts);
-
-		for (int MeshIdx = 0; MeshIdx < Meshes.Num(); MeshIdx++)
-		{
-			FDynamicMeshCollection::FMeshData& MeshData = Meshes[MeshIdx];
-			FDynamicMesh3& Mesh = MeshData.AugMesh;
-			int32 GeometryIdx = Collection.TransformToGeometryIndex[MeshData.TransformIndex];
-			bool bSucceeded = UpdateCollection(MeshData.ToCollection, Mesh, GeometryIdx, Collection, -1);
-			bAllSucceeded &= bSucceeded;
-		}
-
-		return bAllSucceeded;
-	}
-
-	// Update an existing geometry in a collection w/ a new mesh (w/ the same number of faces and vertices!)
-	static bool UpdateCollection(const FTransform& ToCollection, FDynamicMesh3& Mesh, int32 GeometryIdx, FGeometryCollection& Output, int32 InternalMaterialID)
-	{
-		if (!Mesh.IsCompact())
-		{
-			Mesh.CompactInPlace(nullptr);
-		}
-
-		int32 OldVertexCount = Output.VertexCount[GeometryIdx];
-		int32 OldTriangleCount = Output.FaceCount[GeometryIdx];
-
-		int32 NewVertexCount = Mesh.VertexCount();
-		int32 NewTriangleCount = Mesh.TriangleCount();
-
-		if (!ensure(OldVertexCount == NewVertexCount) || !ensure(OldTriangleCount == NewTriangleCount))
-		{
-			return false;
-		}
-
-		int32 VerticesStart = Output.VertexStart[GeometryIdx];
-		int32 FacesStart = Output.FaceStart[GeometryIdx];
-		int32 TransformIdx = Output.TransformIndex[GeometryIdx];
-
-		for (int32 VID = 0; VID < Mesh.MaxVertexID(); VID++)
-		{
-			checkSlow(Mesh.IsVertex(VID)); // mesh is compact
-			int32 CopyToIdx = VerticesStart + VID;
-			Output.Vertex[CopyToIdx] = ToCollection.TransformPosition(FVector(Mesh.GetVertex(VID)));
-			Output.Normal[CopyToIdx] = ToCollection.TransformVectorNoScale(FVector(Mesh.GetVertexNormal(VID)));
-			Output.UV[CopyToIdx] = FVector2D(Mesh.GetVertexUV(VID));
-			FVector3f TangentU, TangentV;
-			UE::PlanarCutInternals::AugmentDynamicMesh::GetTangent(Mesh, VID, TangentU, TangentV);
-			Output.TangentU[CopyToIdx] = ToCollection.TransformVectorNoScale(FVector(TangentU));
-			Output.TangentV[CopyToIdx] = ToCollection.TransformVectorNoScale(FVector(TangentV));
-			Output.Color[CopyToIdx] = FVector(Mesh.GetVertexColor(VID));
-
-			// Bone map is set based on the transform of the new geometry
-			Output.BoneMap[CopyToIdx] = TransformIdx;
-		}
-
-		FIntVector VertexStartOffset(VerticesStart);
-		for (int32 TID = 0; TID < Mesh.MaxTriangleID(); TID++)
-		{
-			checkSlow(Mesh.IsTriangle(TID));
-			int32 CopyToIdx = FacesStart + TID;
-			Output.Visible[CopyToIdx] = UE::PlanarCutInternals::AugmentDynamicMesh::GetVisibility(Mesh, TID);
-			int MaterialID = Mesh.Attributes()->GetMaterialID()->GetValue(TID);
-			Output.MaterialID[CopyToIdx] = MaterialID < 0 ? InternalMaterialID : MaterialID;
-			Output.Indices[CopyToIdx] = FIntVector(Mesh.GetTriangle(TID)) + VertexStartOffset;
-		}
-
-		if (Output.BoundingBox.Num())
-		{
-			Output.BoundingBox[GeometryIdx].Init();
-			for (int32 Idx = VerticesStart; Idx < VerticesStart + Output.VertexCount[GeometryIdx]; ++Idx)
-			{
-				Output.BoundingBox[GeometryIdx] += Output.Vertex[Idx];
-			}
-		}
-
-		return true;
-	}
-
-	static int32 AppendToCollection(const FTransform& ToCollection, FDynamicMesh3& Mesh, double CollisionSampleSpacing, int32 TransformParent, FString BoneName, FGeometryCollection& Output, int32 InternalMaterialID)
-	{
-		if (Mesh.TriangleCount() == 0)
-		{
-			return -1;
-		}
-
-		if (!Mesh.IsCompact())
-		{
-			Mesh.CompactInPlace(nullptr);
-		}
-
-		if (CollisionSampleSpacing > 0)
-		{
-			UE::PlanarCutInternals::AugmentDynamicMesh::AddCollisionSamplesPerComponent(Mesh, CollisionSampleSpacing);
-		}
-
-		int32 NewGeometryStartIdx = Output.FaceStart.Num();
-		int32 OriginalVertexNum = Output.Vertex.Num();
-		int32 OriginalFaceNum = Output.Indices.Num();
-
-		int32 GeometryIdx = Output.AddElements(1, FGeometryCollection::GeometryGroup);
-		int32 TransformIdx = Output.AddElements(1, FGeometryCollection::TransformGroup);
-
-		int32 NumTriangles = Mesh.TriangleCount();
-		int32 NumVertices = Mesh.VertexCount();
-		check(NumTriangles > 0);
-		check(Mesh.IsCompact());
-		Output.FaceCount[GeometryIdx] = NumTriangles;
-		Output.FaceStart[GeometryIdx] = OriginalFaceNum;
-		Output.VertexCount[GeometryIdx] = NumVertices;
-		Output.VertexStart[GeometryIdx] = OriginalVertexNum;
-		Output.TransformIndex[GeometryIdx] = TransformIdx;
-		Output.TransformToGeometryIndex[TransformIdx] = GeometryIdx;
-		if (TransformParent > -1)
-		{
-			Output.BoneName[TransformIdx] = BoneName;
-			Output.BoneColor[TransformIdx] = Output.BoneColor[TransformParent];
-			Output.Parent[TransformIdx] = TransformParent;
-			Output.Children[TransformParent].Add(TransformIdx);
-			Output.SimulationType[TransformParent] = FGeometryCollection::ESimulationTypes::FST_Clustered;
-		}
-		Output.Transform[TransformIdx] = FTransform::Identity;
-		Output.SimulationType[TransformIdx] = FGeometryCollection::ESimulationTypes::FST_Rigid;
-
-		int32 FacesStart = Output.AddElements(NumTriangles, FGeometryCollection::FacesGroup);
-		int32 VerticesStart = Output.AddElements(NumVertices, FGeometryCollection::VerticesGroup);
-
-		for (int32 VID = 0; VID < Mesh.MaxVertexID(); VID++)
-		{
-			checkSlow(Mesh.IsVertex(VID)); // mesh is compact
-			int32 CopyToIdx = VerticesStart + VID;
-			Output.Vertex[CopyToIdx] = ToCollection.TransformPosition(FVector(Mesh.GetVertex(VID)));
-			Output.Normal[CopyToIdx] = ToCollection.TransformVectorNoScale(FVector(Mesh.GetVertexNormal(VID)));
-			Output.UV[CopyToIdx] = FVector2D(Mesh.GetVertexUV(VID));
-			FVector3f TangentU, TangentV;
-			UE::PlanarCutInternals::AugmentDynamicMesh::GetTangent(Mesh, VID, TangentU, TangentV);
-			Output.TangentU[CopyToIdx] = ToCollection.TransformVectorNoScale(FVector(TangentU));
-			Output.TangentV[CopyToIdx] = ToCollection.TransformVectorNoScale(FVector(TangentV));
-			Output.Color[CopyToIdx] = FVector(Mesh.GetVertexColor(VID));
-
-			// Bone map is set based on the transform of the new geometry
-			Output.BoneMap[CopyToIdx] = TransformIdx;
-		}
-
-		FIntVector VertexStartOffset(VerticesStart);
-		for (int32 TID = 0; TID < Mesh.MaxTriangleID(); TID++)
-		{
-			checkSlow(Mesh.IsTriangle(TID));
-			int32 CopyToIdx = FacesStart + TID;
-			Output.Visible[CopyToIdx] = UE::PlanarCutInternals::AugmentDynamicMesh::GetVisibility(Mesh, TID);
-			int MaterialID = Mesh.Attributes()->GetMaterialID()->GetValue(TID);
-			Output.MaterialID[CopyToIdx] = MaterialID < 0 ? InternalMaterialID : MaterialID;
-			Output.Indices[CopyToIdx] = FIntVector(Mesh.GetTriangle(TID)) + VertexStartOffset;
-		}
-
-		if (Output.BoundingBox.Num())
-		{
-			Output.BoundingBox[GeometryIdx].Init();
-			for (int32 Idx = OriginalVertexNum; Idx < Output.Vertex.Num(); ++Idx)
-			{
-				Output.BoundingBox[GeometryIdx] += Output.Vertex[Idx];
-			}
-		}
-
-		return GeometryIdx;
-	}
-};
+using namespace UE::PlanarCut;
 
 
 // logic from FMeshUtility::GenerateGeometryCollectionFromBlastChunk, sets material IDs based on construction pattern that external materials have even IDs and are matched to internal materials at InternalID = ExternalID+1
@@ -2910,7 +662,7 @@ int32 CutWithMesh(
 	}
 
 	FDynamicMesh3 DynamicCuttingMesh; // version of mesh that is split apart at seams to be compatible w/ geometry collection, with corresponding attributes set
-	UE::PlanarCutInternals::AugmentDynamicMesh::Augment(DynamicCuttingMesh);
+	SetGeometryCollectionAttributes(DynamicCuttingMesh);
 	// Note: This conversion will likely go away, b/c I plan to switch over to doing the boolean operations on the fuller rep, but the code can be adapted
 	//		 to the dynamic mesh -> geometry collection conversion phase, as this same splitting will then need to happen there.
 	if (ensure(FullMesh.HasAttributes() && FullMesh.Attributes()->NumUVLayers() >= 1 && FullMesh.Attributes()->NumNormalLayers() == 3))
@@ -2957,7 +709,7 @@ int32 CutWithMesh(
 
 					int OutVID = OutMesh.AppendVertex(Info);
 					OutTri[VOnT] = OutVID;
-					UE::PlanarCutInternals::AugmentDynamicMesh::SetTangent(OutMesh, OutVID, Normal, NTB[1]->GetElement(ElIDs.B), NTB[2]->GetElement(ElIDs.C));
+					AugmentedDynamicMesh::SetTangent(OutMesh, OutVID, Normal, NTB[1]->GetElement(ElIDs.B), NTB[2]->GetElement(ElIDs.C));
 					ElIDsToVID.Add(ElIDs, OutVID);
 				}
 			});
@@ -2971,7 +723,7 @@ int32 CutWithMesh(
 			if (ensure(AddedTID > -1))
 			{
 				OutMaterialID->SetValue(AddedTID, -1); // just use a single negative material ID by convention to indicate internal material
-				UE::PlanarCutInternals::AugmentDynamicMesh::SetVisibility(OutMesh, AddedTID, true);
+				AugmentedDynamicMesh::SetVisibility(OutMesh, AddedTID, true);
 			}
 		}
 	}
@@ -3014,7 +766,7 @@ void RecomputeNormalsAndTangents(bool bOnlyTangents, FGeometryCollection& Collec
 	for (int MeshIdx = 0; MeshIdx < MeshCollection.Meshes.Num(); MeshIdx++)
 	{
 		FDynamicMesh3& Mesh = MeshCollection.Meshes[MeshIdx].AugMesh;
-		UE::PlanarCutInternals::AugmentDynamicMesh::ComputeTangents(Mesh, bOnlyOddMaterials, WhichMaterials, !bOnlyTangents);
+		AugmentedDynamicMesh::ComputeTangents(Mesh, bOnlyOddMaterials, WhichMaterials, !bOnlyTangents);
 	}
 
 	MeshCollection.UpdateAllCollections(Collection);
@@ -3056,7 +808,7 @@ void ConvertToMeshDescription(
 	FDynamicMeshCollection MeshCollection(&Collection, TransformIndices, CellsToWorld);
 	
 	FDynamicMesh3 CombinedMesh;
-	UE::PlanarCutInternals::AugmentDynamicMesh::Augment(CombinedMesh);
+	SetGeometryCollectionAttributes(CombinedMesh);
 	CombinedMesh.Attributes()->EnableTangents();
 
 	int32 NumMeshes = MeshCollection.Meshes.Num();
@@ -3073,8 +825,8 @@ void ConvertToMeshDescription(
 		}
 
 		FMeshNormals::InitializeOverlayToPerVertexNormals(Mesh.Attributes()->PrimaryNormals(), true);
-		UE::PlanarCutInternals::AugmentDynamicMesh::InitializeOverlayToPerVertexUVs(Mesh);
-		UE::PlanarCutInternals::AugmentDynamicMesh::InitializeOverlayToPerVertexTangents(Mesh);
+		AugmentedDynamicMesh::InitializeOverlayToPerVertexUVs(Mesh);
+		AugmentedDynamicMesh::InitializeOverlayToPerVertexTangents(Mesh);
 
 		FMergeCoincidentMeshEdges EdgeMerge(&Mesh);
 		EdgeMerge.Apply();
