PlanarCut
-- remove pre-cut geometry instead of hiding it to correct issues with nanite scene proxy rendering.

diff --git a/Engine/Plugins/Experimental/PlanarCutPlugin/Source/PlanarCut/Private/GeometryMeshConversion.cpp b/Engine/Plugins/Experimental/PlanarCutPlugin/Source/PlanarCut/Private/GeometryMeshConversion.cpp
--- a/Engine/Plugins/Experimental/PlanarCutPlugin/Source/PlanarCut/Private/GeometryMeshConversion.cpp
+++ b/Engine/Plugins/Experimental/PlanarCutPlugin/Source/PlanarCut/Private/GeometryMeshConversion.cpp
@@ -1692,6 +1692,7 @@ int32 FDynamicMeshCollection::CutWithMultiplePlanes(
 	TArray<int32> ToCutIdxToGeometryIdx;  ToCutIdxToGeometryIdx.Init(-1, ToCut.Num());
 	TArray<int32> ToCutIndices;
 	int32 FirstCreatedIndex = -1;
+	TArray<int32> GeometryForRemoval;
 	for (FMeshData& MeshData : Meshes)
 	{
 		int32 GeometryIdx = Collection->TransformToGeometryIndex[MeshData.TransformIndex];
@@ -1705,8 +1706,9 @@ int32 FDynamicMeshCollection::CutWithMultiplePlanes(
 			continue;
 		}
 
-		// hide old parent geometry
-		SetVisibility(*Collection, GeometryIdx, false);
+		// remove old parent geometry
+		GeometryForRemoval.Add(GeometryIdx);
+		Collection->TransformToGeometryIndex[MeshData.TransformIndex] = INDEX_NONE;
 
 		// add newly created geometry as children
 		int32 SubPartIdx = 0;
@@ -1734,6 +1736,13 @@ int32 FDynamicMeshCollection::CutWithMultiplePlanes(
 		}
 	}
 
+	// remove superfluous geometry
+	if (GeometryForRemoval.Num() > 0)
+	{
+		GeometryForRemoval.Sort();
+		Collection->RemoveElements(FGeometryCollection::GeometryGroup, GeometryForRemoval);
+	}
+
 	return FirstCreatedIndex;
 }
 
@@ -1744,6 +1753,7 @@ int32 FDynamicMeshCollection::CutWithCellMeshes(const FInternalSurfaceMaterials&
 	int32 FirstIdx = -1;
 	int BadCount = 0;
 	bool bHasProximity = Collection->HasAttribute("Proximity", FGeometryCollection::GeometryGroup);
+	TArray<int32> GeometryForRemoval;
 	for (FMeshData& Surface : Meshes)
 	{
 		int32 GeometryIdx = Collection->TransformToGeometryIndex[Surface.TransformIndex];
@@ -1890,11 +1900,21 @@ int32 FDynamicMeshCollection::CutWithCellMeshes(const FInternalSurfaceMaterials&
 					}
 				}
 			}
-			// turn off old geom visibility (preferred default behavior)
-			SetVisibility(*Collection, GeometryIdx, false);
+			
+			// remove old geom
+			GeometryForRemoval.Add(GeometryIdx);
+			Collection->TransformToGeometryIndex[Surface.TransformIndex] = INDEX_NONE;
+
 		}
 	}
 
+	// remove superfluous geometry
+	if (GeometryForRemoval.Num() > 0)
+	{
+		GeometryForRemoval.Sort();
+		Collection->RemoveElements(FGeometryCollection::GeometryGroup, GeometryForRemoval);
+	}
+
 	return FirstIdx;
 }
 
