Add more texturing options to fracture AutoUV tool

#rnx

diff --git a/Engine/Plugins/Experimental/ChaosEditor/Source/FractureEditor/Private/FractureToolUV.cpp b/Engine/Plugins/Experimental/ChaosEditor/Source/FractureEditor/Private/FractureToolUV.cpp
--- a/Engine/Plugins/Experimental/ChaosEditor/Source/FractureEditor/Private/FractureToolUV.cpp
+++ b/Engine/Plugins/Experimental/ChaosEditor/Source/FractureEditor/Private/FractureToolUV.cpp
@@ -165,8 +165,6 @@ bool UFractureToolAutoUV::SaveGeneratedTexture(UTexture2D* GeneratedTexture, FSt
 		}
 	}
 
-	FString UseBaseName = ObjectBaseName;
-
 	// create new package
 	FString UniqueAssetName;
 	FString UniquePackageName;
@@ -174,7 +172,6 @@ bool UFractureToolAutoUV::SaveGeneratedTexture(UTexture2D* GeneratedTexture, FSt
 	FAssetToolsModule& AssetToolsModule = FModuleManager::LoadModuleChecked<FAssetToolsModule>("AssetTools");
 	AssetToolsModule.Get().CreateUniqueAssetName(PackageFolderPath + TEXT("/") + ObjectBaseName, TEXT(""), UniquePackageName, UniqueAssetName);
 
-	UPackage* Package = CreatePackage(*UniquePackageName);
 	UPackage* AssetPackage = CreatePackage(*UniquePackageName);
 
 	// move texture from Transient package to real package
@@ -218,11 +215,26 @@ int32 UFractureToolAutoUV::ExecuteFracture(const FFractureToolContext& FractureC
 		FImageDimensions Dimensions(OutputRes, OutputRes);
 		TextureBuilder.Initialize(FTexture2DBuilder::ETextureType::Color, Dimensions);
 
-		UE::Geometry::TImageBuilder<FVector3f> ImageBuilder;
+		UE::Geometry::TImageBuilder<FVector4f> ImageBuilder;
 		ImageBuilder.SetDimensions(Dimensions);
-		ImageBuilder.Clear(FVector3f(0, 0, 0));
-
-		UE::PlanarCut::TextureInternalSurfaces(Collection, AutoUVSettings->MaxDistance, FMath::CeilToInt(AutoUVSettings->GutterSize), ImageBuilder);
+		ImageBuilder.Clear(FVector4f(0, 0, 0, 0));
+
+		typedef UE::PlanarCut::EBakeAttributes EBakeAttributes;
+		// Note: Ordering of these attributes should match the order and comments in the AutoUVSettings struct
+		//		 Update the order and comments there if you change the ordering here.
+		FIndex4i Attributes(
+			AutoUVSettings->bDistToOuter ? (int32)EBakeAttributes::DistanceToExternal : 0, 
+			AutoUVSettings->bAmbientOcclusion ? (int32)EBakeAttributes::AmbientOcclusion : 0, 
+			AutoUVSettings->bZNormal ? (int32)EBakeAttributes::NormalZ : 0,
+			AutoUVSettings->bZPosition ? (int32)EBakeAttributes::PositionZ : 0
+		);
+		UE::PlanarCut::FTextureAttributeSettings AttribSettings;
+		AttribSettings.ToExternal_MaxDistance = AutoUVSettings->MaxDistance;
+		AttribSettings.AO_Rays = AutoUVSettings->OcclusionRays;
+		AttribSettings.bAO_Blur = AutoUVSettings->bGaussianBlur;
+		AttribSettings.AO_BlurRadius = AutoUVSettings->BlurRadius;
+		AttribSettings.bNormalZ_TakeAbs = AutoUVSettings->bUseAbsoluteValue;
+		UE::PlanarCut::TextureInternalSurfaces(Collection, FMath::CeilToInt(AutoUVSettings->GutterSize), Attributes, AttribSettings, ImageBuilder);
 
 		UVTask.EnterProgressFrame(1, LOCTEXT("SavingTexture", "Saving result"));
 
diff --git a/Engine/Plugins/Experimental/ChaosEditor/Source/FractureEditor/Private/FractureToolUV.h b/Engine/Plugins/Experimental/ChaosEditor/Source/FractureEditor/Private/FractureToolUV.h
--- a/Engine/Plugins/Experimental/ChaosEditor/Source/FractureEditor/Private/FractureToolUV.h
+++ b/Engine/Plugins/Experimental/ChaosEditor/Source/FractureEditor/Private/FractureToolUV.h
@@ -44,10 +44,6 @@ public:
 	UPROPERTY(EditAnywhere, Category = MapSettings, meta = (ClampMin = "1", ClampMax = "10", UIMax = "4"))
 	int32 GutterSize = 2;
 
-	/** Max distance to search for the outer mesh surface */
-	UPROPERTY(EditAnywhere, Category = MapSettings, meta = (UIMin = "1", UIMax = "100", ClampMin = ".01", ClampMax = "1000"))
-	double MaxDistance = 50;
-
 	/** The resulting automatically-generated texture map */
 	UPROPERTY(VisibleAnywhere, Category = MapSettings)
 	UTexture2D* Result;
@@ -55,6 +51,43 @@ public:
 	/** Whether to prompt user for an asset name for each generated texture, or automatically place them next to the source geometry collections */
 	UPROPERTY(EditAnywhere, Category = MapSettings)
 	bool bPromptToSave = true;
+
+	/** Bake the distance to the external surface to a texture channel (red) */
+	UPROPERTY(EditAnywhere, Category = AttributesToBake)
+	bool bDistToOuter = true;
+
+	/** Bake the ambient occlusion of each bone (considered separately) to a texture channel (green) */
+	UPROPERTY(EditAnywhere, Category = AttributesToBake)
+	bool bAmbientOcclusion = true;
+
+	/** Bake the Z-component of the normal to a texture channel (blue) */
+	UPROPERTY(EditAnywhere, Category = AttributesToBake)
+	bool bZNormal = true;
+
+	/** Bake the Z-component of the position (relative to the overall bounding box) to a texture channel (alpha) */
+	UPROPERTY(EditAnywhere, Category = AttributesToBake)
+	bool bZPosition = false;
+
+	/** Max distance to search for the outer mesh surface */
+	UPROPERTY(EditAnywhere, Category = DistToOuterSettings, meta = (EditCondition = "bDistToOuter", EditConditionHides, UIMin = "1", UIMax = "100", ClampMin = ".01", ClampMax = "1000"))
+	double MaxDistance = 100;
+
+	/** Number of occlusion rays */
+	UPROPERTY(EditAnywhere, Category = AmbientOcclusionSettings, meta = (EditCondition = "bAmbientOcclusion", EditConditionHides, UIMin = "1", UIMax = "1024", ClampMin = "0", ClampMax = "50000"))
+	int OcclusionRays = 16;
+
+	/** Whether or not to apply Gaussian Blur to computed AO Map (recommended) */
+	UPROPERTY(EditAnywhere, Category = AmbientOcclusionSettings, meta = (EditCondition = "bAmbientOcclusion", EditConditionHides))
+	bool bGaussianBlur = true;
+
+	/** Pixel Radius of Gaussian Blur Kernel */
+	UPROPERTY(EditAnywhere, Category = AmbientOcclusionSettings, meta = (EditCondition = "bAmbientOcclusion && bGaussianBlur", EditConditionHides, UIMin = "0", UIMax = "10.0", ClampMin = "0", ClampMax = "100.0"))
+	double BlurRadius = 2.25;
+
+	/** Whether to use the absolute value of the Z-Component of the normal */
+	UPROPERTY(EditAnywhere, Category = ZNormalSettings, meta = (EditCondition = "bZNormal", EditConditionHides))
+	bool bUseAbsoluteValue = true;
+
 };
 
 
diff --git a/Engine/Plugins/Experimental/PlanarCutPlugin/Source/PlanarCut/Private/FractureAutoUV.cpp b/Engine/Plugins/Experimental/PlanarCutPlugin/Source/PlanarCut/Private/FractureAutoUV.cpp
--- a/Engine/Plugins/Experimental/PlanarCutPlugin/Source/PlanarCut/Private/FractureAutoUV.cpp
+++ b/Engine/Plugins/Experimental/PlanarCutPlugin/Source/PlanarCut/Private/FractureAutoUV.cpp
@@ -2,20 +2,27 @@
 #include "FractureAutoUV.h"
 #include "PlanarCutPlugin.h"
 
+#include "GeometryMeshConversion.h"
+
 #include "Async/ParallelFor.h"
 
 #include "GeometryCollection/GeometryCollectionAlgo.h"
 
 #include "DynamicMesh3.h"
+#include "MeshNormals.h"
 #include "Parameterization/DynamicMeshUVEditor.h"
+#include "Sampling/MeshImageBakingCache.h"
 
 #include "Parameterization/UVPacking.h"
+#include "Sampling/MeshOcclusionMapBaker.h"
 #include "DisjointSet.h"
 
 #include "Image/ImageOccupancyMap.h"
 #include "VectorUtil.h"
 #include "FrameTypes.h"
 
+#include "Templates/PimplPtr.h"
+
 #if WITH_EDITOR
 #include "Misc/ScopedSlowTask.h"
 #endif
@@ -95,6 +102,28 @@ namespace UE
 
 namespace UE { namespace PlanarCut {
 
+
+namespace {
+	inline bool IsTriActive(bool bIsVisible, int32 MaterialID, const TSet<int32>& InsideMaterials, bool bActivateInsideTriangles, bool bOddMaterialAreInside)
+	{
+		if (!bIsVisible)
+		{
+			return false;
+		}
+
+		if (bOddMaterialAreInside && ((MaterialID % 2) == 0) == bActivateInsideTriangles)
+		{
+			return false;
+		}
+		if (InsideMaterials.Num() > 0 && (InsideMaterials.Contains(MaterialID) != bActivateInsideTriangles))
+		{
+			return false;
+		}
+
+		return true;
+	}
+}
+
 /**
  * Shared helper method to set 'active' subset of triangles for a collection, to be considered for texturing
  * @return number of active triangles
@@ -113,23 +142,12 @@ int32 SetActiveTriangles(FGeometryCollection* Collection, TArray<bool>& ActiveTr
 	int32 NumTriangles = 0;
 	for (int TID = 0; TID < Collection->Indices.Num(); TID++)
 	{
-		if (!Collection->Visible[TID])
-		{
-			continue;
-		}
-
-		int32 MaterialID = Collection->MaterialID[TID];
-		if (bOddMaterialAreInside && ((MaterialID % 2) == 0) == bActivateInsideTriangles)
+		bool bIsActive = IsTriActive(Collection->Visible[TID], Collection->MaterialID[TID], InsideMaterials, bActivateInsideTriangles, bOddMaterialAreInside);
+		if (bIsActive)
 		{
-			continue;
+			ActiveTrianglesOut[TID] = true;
+			NumTriangles++;
 		}
-		if (WhichMaterialsAreInside.Num() > 0 && (InsideMaterials.Contains(MaterialID) != bActivateInsideTriangles))
-		{
-			continue;
-		}
-
-		ActiveTrianglesOut[TID] = true;
-		NumTriangles++;
 	}
 
 	return NumTriangles;
@@ -143,6 +161,7 @@ struct FGeomMesh : public UE::Geometry::FUVPacker::IUVMeshView
 	const TArrayView<bool> ActiveTriangles;
 	int32 NumTriangles;
 	TArray<FVector3d> GlobalVertices; // vertices transformed to global space
+	TArray<FVector3f> GlobalNormals; // normals transformed to global space
 
 	// Construct a mesh from a geometry collection and a mask of active triangles
 	FGeomMesh(FGeometryCollection* Collection, TArrayView<bool> ActiveTriangles, int32 NumTriangles) 
@@ -164,8 +183,10 @@ struct FGeomMesh : public UE::Geometry::FUVPacker::IUVMeshView
 		TArray<FTransform> GlobalTransformArray;
 		GeometryCollectionAlgo::GlobalMatrices(Collection->Transform, Collection->Parent, GlobalTransformArray);
 		GlobalVertices.SetNum(Collection->Vertex.Num());
+		GlobalNormals.SetNum(Collection->Vertex.Num());
 		for (int32 Idx = 0; Idx < GlobalVertices.Num(); Idx++)
 		{
+			GlobalNormals[Idx] = (FVector3f)GlobalTransformArray[Collection->BoneMap[Idx]].TransformVectorNoScale(Collection->Normal[Idx]);
 			GlobalVertices[Idx] = (FVector3d)GlobalTransformArray[Collection->BoneMap[Idx]].TransformPosition(Collection->Vertex[Idx]);
 		}
 	}
@@ -243,6 +264,15 @@ struct FGeomMesh : public UE::Geometry::FUVPacker::IUVMeshView
 		V1 = GlobalVertices[TriRaw.Y];
 		V2 = GlobalVertices[TriRaw.Z];
 	}
+
+	/// End of interfaces
+	
+	FVector3f GetInterpolatedNormal(int TID, const FVector3d& Bary)
+	{
+		FIntVector TriRaw = Collection->Indices[TID];
+		FVector3f Normal = (FVector3f) (Bary.X * GlobalNormals[TriRaw.X] + Bary.Y * GlobalNormals[TriRaw.Y] + Bary.Z * GlobalNormals[TriRaw.Z]);
+		return Normalized(Normal);
+	}
 };
 
 /// Like FGeomMesh, but the UVs are given as the vertices rather than accessed separately
@@ -438,11 +468,13 @@ bool UVLayout(
 	return Packer.StandardPack(&UVMesh, UVIslands);
 }
 
+
 void TextureInternalSurfaces(
 	FGeometryCollection& Collection,
-	double MaxDistance,
 	int32 GutterSize,
-	TImageBuilder<FVector3f>& TextureOut,
+	FIndex4i BakeAttributes,
+	const FTextureAttributeSettings& AttributeSettings,
+	TImageBuilder<FVector4f>& TextureOut,
 	bool bOnlyOddMaterials,
 	TArrayView<int32> WhichMaterials
 )
@@ -463,33 +495,180 @@ void TextureInternalSurfaces(
 	FGeomMesh OutsideMesh(InsideMesh, OutsideTriangles, NumOutsideTris);
 	TMeshAABBTree3<FGeomMesh> OutsideSpatial(&OutsideMesh);
 
+	int DistanceToExternalIdx = BakeAttributes.IndexOf((int)EBakeAttributes::DistanceToExternal);
+	int AmbientIdx = BakeAttributes.IndexOf((int)EBakeAttributes::AmbientOcclusion);
+	int NormalZIdx = BakeAttributes.IndexOf((int)EBakeAttributes::NormalZ);
+	int PosZIdx = BakeAttributes.IndexOf((int)EBakeAttributes::PositionZ);
+	if (OutsideSpatial.GetBoundingBox().Depth() == 0)
+	{
+		PosZIdx = -1; // if everything has same Z, can just leave PosZ values as default value
+	}
+
+
+	bool bNeedsDynamicMeshes = AmbientIdx > -1;
+
+	TArray<int32> TransformIndices;
+	if (bNeedsDynamicMeshes)
+	{
+		for (int32 GeomIdx = 0; GeomIdx < Collection.TransformIndex.Num(); GeomIdx++)
+		{
+			TransformIndices.Add(Collection.TransformIndex[GeomIdx]);
+		}
+	}
+	FDynamicMeshCollection CollectionMeshes(&Collection, TransformIndices, FTransform::Identity, false);
+	if (bNeedsDynamicMeshes)
+	{
+		TSet<int32> InsideMaterials;
+		for (int32 MatID : WhichMaterials)
+		{
+			InsideMaterials.Add(MatID);
+		}
+
+		// To bake only internal faces, unset texture on everything else
+		for (int MeshIdx = 0; MeshIdx < CollectionMeshes.Meshes.Num(); MeshIdx++)
+		{
+			int32 TransformIdx = CollectionMeshes.Meshes[MeshIdx].TransformIndex;
+			FDynamicMesh3& Mesh = CollectionMeshes.Meshes[MeshIdx].AugMesh;
+			FMeshNormals::InitializeOverlayToPerVertexNormals(Mesh.Attributes()->PrimaryNormals(), true);
+			AugmentedDynamicMesh::InitializeOverlayToPerVertexUVs(Mesh);
+			FDynamicMeshUVOverlay* UV = Mesh.Attributes()->PrimaryUV();
+			for (int TID : Mesh.TriangleIndicesItr())
+			{
+				bool bIsInsideTri = IsTriActive(
+					AugmentedDynamicMesh::GetVisibility(Mesh, TID), 
+					Mesh.Attributes()->GetMaterialID()->GetValue(TID), InsideMaterials, true, bOnlyOddMaterials);
+				if (!bIsInsideTri)
+				{
+					UV->UnsetTriangle(TID);
+				}
+			}
+		}
+	}
+	
+	TArray64<float> AmbientValues; // buffer to fill with ambient occlusion values
+	if (AmbientIdx > -1)
+	{
+		FImageDimensions Dimensions = TextureOut.GetDimensions();
+		int64 TexelsNum = Dimensions.Num();
+		AmbientValues.SetNumZeroed(TexelsNum);
+		for (int MeshIdx = 0; MeshIdx < CollectionMeshes.Meshes.Num(); MeshIdx++)
+		{
+			int32 TransformIdx = CollectionMeshes.Meshes[MeshIdx].TransformIndex;
+			FDynamicMesh3& Mesh = CollectionMeshes.Meshes[MeshIdx].AugMesh;
+			FDynamicMeshUVOverlay* UV = Mesh.Attributes()->PrimaryUV();
+
+			FDynamicMeshAABBTree3 Spatial(&Mesh);
+
+			FMeshImageBakingCache BakeCache;
+			BakeCache.SetDetailMesh(&Mesh, &Spatial);
+			BakeCache.SetBakeTargetMesh(&Mesh);
+			BakeCache.SetDimensions(Dimensions);
+			BakeCache.SetUVLayer(0);
+			// thickness is used for raycasting correspondences between detail and target mesh
+			// TODO: make cache aware that detail mesh == target mesh so it doesn't have to do this?
+			BakeCache.SetThickness(KINDA_SMALL_NUMBER);
+			BakeCache.ValidateCache();
+
+			FMeshOcclusionMapBaker OcclusionBaker;
+			OcclusionBaker.SetCache(&BakeCache);
+			
+			// set NormalSpace to Object to avoid need to compute tangents for bent normals computation by the AO baker, which we don't need here.
+			// TODO: consider adding a toggle to fully disable bent normals computation in FMeshOcclusionMapBaker
+			OcclusionBaker.NormalSpace = FMeshOcclusionMapBaker::ESpace::Object;
+			OcclusionBaker.NumOcclusionRays = AttributeSettings.AO_Rays;
+			OcclusionBaker.MaxDistance = AttributeSettings.AO_MaxDistance == 0 ? TNumericLimits<double>::Max() : AttributeSettings.AO_MaxDistance;
+			OcclusionBaker.BiasAngleDeg = AttributeSettings.AO_BiasAngleDeg;
+			OcclusionBaker.BlurRadius = AttributeSettings.bAO_Blur ? AttributeSettings.AO_BlurRadius : 0.0;
+
+			OcclusionBaker.Bake();
+
+			const TUniquePtr<TImageBuilder<FVector3f>>& Result = OcclusionBaker.GetResult(FMeshOcclusionMapBaker::EResult::AmbientOcclusion);
+			// copy scalar ambient values out to the accumulated ambient buffer
+			for (int64 LinearIdx = 0; LinearIdx < TexelsNum; LinearIdx++)
+			{
+				float& AmbientOut = AmbientValues[LinearIdx];
+				AmbientOut = FMathf::Max(AmbientOut, Result->GetPixel(LinearIdx).X);
+			}
+		}
+	}
+
 	ParallelFor(OccupancyMap.Dimensions.GetHeight(),
-		[&MaxDistance, &TextureOut, &UVMesh, &OccupancyMap, &InsideMesh, &OutsideSpatial](int32 Y)
+		[&AttributeSettings, &TextureOut, &UVMesh, &OccupancyMap, &InsideMesh, &OutsideSpatial,
+		 &DistanceToExternalIdx, &AmbientIdx, &NormalZIdx, &PosZIdx](int32 Y)
 	{
 		for (int32 X = 0; X < OccupancyMap.Dimensions.GetWidth(); X++)
 		{
 			int64 LinearCoord = OccupancyMap.Dimensions.GetIndex(X, Y);
 			if (OccupancyMap.IsInterior(LinearCoord))
 			{
+				FVector4f OutColor = FVector4f::One();
+
 				int32 TID = OccupancyMap.TexelQueryTriangle[LinearCoord];
 				FVector2f UV = OccupancyMap.TexelQueryUV[LinearCoord];
 				// Note this is the slowest way to get barycentric coordinates out of a point and a 2D triangle,
 				//   ... but it's the one we currently have that's robust to (and gives good answers on) degenerate triangles
 				FDistPoint3Triangle3d DistQuery = TMeshQueries<FGeomFlatUVMesh>::TriangleDistance(UVMesh, TID, FVector3d(UV.X, UV.Y, 0));
 				FVector3d Bary = DistQuery.TriangleBaryCoords;
-				FTriangle3d Tri;
-				InsideMesh.GetTriVertices(TID, Tri.V[0], Tri.V[1], Tri.V[2]);
-				FVector3d InsidePoint = Tri.BarycentricPoint(Bary);
-				double DistanceSq;
-				OutsideSpatial.FindNearestTriangle(InsidePoint, DistanceSq, MaxDistance);
-				float PercentDistance = FMathf::Min(1.0f, FMathf::Sqrt(DistanceSq) / (float)MaxDistance);
-				ensure(FMath::IsFinite(PercentDistance));
-
-				TextureOut.SetPixel(LinearCoord, FVector3f(PercentDistance, 0, 1));
+				FVector3f Normal(0, 0, 1);
+				bool bNeedsNormal = NormalZIdx > -1 || AmbientIdx > -1;
+				if (bNeedsNormal)
+				{
+					Normal = InsideMesh.GetInterpolatedNormal(TID, Bary);
+				}
+
+				if (DistanceToExternalIdx > -1 || PosZIdx > -1)
+				{
+					FTriangle3d Tri;
+					InsideMesh.GetTriVertices(TID, Tri.V[0], Tri.V[1], Tri.V[2]);
+					FVector3d InsidePoint = Tri.BarycentricPoint(Bary);
+
+					if (DistanceToExternalIdx > -1)
+					{
+						double DistanceSq;
+						OutsideSpatial.FindNearestTriangle(InsidePoint, DistanceSq, AttributeSettings.ToExternal_MaxDistance);
+						float PercentDistance = FMathf::Min(1.0f, FMathf::Sqrt(DistanceSq) / (float)AttributeSettings.ToExternal_MaxDistance);
+						checkSlow(FMath::IsFinite(PercentDistance));
+						OutColor[DistanceToExternalIdx] = PercentDistance;
+					}
+					if (PosZIdx > -1)
+					{
+						double MinZ = OutsideSpatial.GetBoundingBox().Min.Z, MaxZ = OutsideSpatial.GetBoundingBox().Max.Z;
+						checkSlow(MinZ != MaxZ);
+						float PercentHeight = (float)(InsidePoint.Z - MinZ) / (MaxZ - MinZ);
+						OutColor[PosZIdx] = PercentHeight;
+					}
+				}
+				if (NormalZIdx > -1)
+				{
+					float NormalZ = Normal.Z;
+					if (AttributeSettings.bNormalZ_TakeAbs)
+					{
+						NormalZ = FMathf::Abs(NormalZ);
+					}
+					else
+					{
+						NormalZ = (1 + NormalZ) * .5; // compress from [-1,1] to [0,1] range
+					}
+					OutColor[NormalZIdx] = NormalZ;
+				}
+
+				TextureOut.SetPixel(LinearCoord, OutColor);
 			}
 		}
 	});
 
+	// Copy the ambient occlusion values to the appropriate channel
+	if (AmbientIdx > -1)
+	{
+		int64 TexelsNum = OccupancyMap.Dimensions.Num();
+		for (int64 LinearIdx = 0; LinearIdx < TexelsNum; LinearIdx++)
+		{
+			FVector4f Pixel = TextureOut.GetPixel(LinearIdx);
+			Pixel[AmbientIdx] = AmbientValues[LinearIdx];
+			TextureOut.SetPixel(LinearIdx, Pixel);
+		}
+	}
+
 	for (TTuple<int64, int64>& GutterToInside : OccupancyMap.GutterTexels)
 	{
 		TextureOut.CopyPixel(GutterToInside.Value, GutterToInside.Key);
diff --git a/Engine/Plugins/Experimental/PlanarCutPlugin/Source/PlanarCut/Private/GeometryMeshConversion.h b/Engine/Plugins/Experimental/PlanarCutPlugin/Source/PlanarCut/Private/GeometryMeshConversion.h
--- a/Engine/Plugins/Experimental/PlanarCutPlugin/Source/PlanarCut/Private/GeometryMeshConversion.h
+++ b/Engine/Plugins/Experimental/PlanarCutPlugin/Source/PlanarCut/Private/GeometryMeshConversion.h
@@ -29,6 +29,7 @@ void PLANARCUT_API SetGeometryCollectionAttributes(UE::Geometry::FDynamicMesh3&
 namespace AugmentedDynamicMesh
 {
 	void PLANARCUT_API SetVisibility(UE::Geometry::FDynamicMesh3& Mesh, int TID, bool bIsVisible);
+	bool PLANARCUT_API GetVisibility(const UE::Geometry::FDynamicMesh3& Mesh, int TID);
 	void PLANARCUT_API SetTangent(UE::Geometry::FDynamicMesh3& Mesh, int VID, FVector3f Normal, FVector3f TangentU, FVector3f TangentV);
 	void PLANARCUT_API GetTangent(const UE::Geometry::FDynamicMesh3& Mesh, int VID, FVector3f& U, FVector3f& V);
 	void PLANARCUT_API InitializeOverlayToPerVertexUVs(UE::Geometry::FDynamicMesh3& Mesh);
diff --git a/Engine/Plugins/Experimental/PlanarCutPlugin/Source/PlanarCut/Public/FractureAutoUV.h b/Engine/Plugins/Experimental/PlanarCutPlugin/Source/PlanarCut/Public/FractureAutoUV.h
--- a/Engine/Plugins/Experimental/PlanarCutPlugin/Source/PlanarCut/Public/FractureAutoUV.h
+++ b/Engine/Plugins/Experimental/PlanarCutPlugin/Source/PlanarCut/Public/FractureAutoUV.h
@@ -28,22 +28,47 @@ bool PLANARCUT_API UVLayout(
 );
 
 
+// Different attributes we can bake
+enum class EBakeAttributes : int32
+{
+	None,
+	DistanceToExternal,
+	AmbientOcclusion,
+	//Curvature, // TODO: add a smoothed curvature attribute
+	NormalZ,
+	PositionZ
+};
+
+struct FTextureAttributeSettings
+{
+	double ToExternal_MaxDistance = 100.0;
+	int AO_Rays = 32;
+	double AO_BiasAngleDeg = 15.0;
+	bool bAO_Blur = true;
+	double AO_BlurRadius = 2.5;
+	double AO_MaxDistance = 0.0; // 0.0 is interpreted as TNumericLimits<double>::Max()
+	bool bNormalZ_TakeAbs = true;
+};
+
+
 /**
  * Generate a texture for internal faces based on depth inside surface
  * TODO: add options to texture based on other quantities
  *
  * @param Collection		The collection to be create a new texture for
- * @param MaxDistance		Maximum distance to search for 'outside' surface, when computing depth inside surface
  * @param GutterSize		Number of texels to fill outside of UV island borders (values are copied from nearest inside pt)
+ * @param BakeAttributes	Which attributes to bake into which color channel
+ * @param AttributeSettings	Settings for the BakeAttributes
  * @param TextureOut		Texture to write to
  * @param bOnlyOddMaterials	If true, restrict UV island layout to odd-numbered material IDs
  * @param WhichMaterials	If non-empty, restrict UV island layout to only the listed material IDs
  */
 void PLANARCUT_API TextureInternalSurfaces(
 	FGeometryCollection& Collection,
-	double MaxDistance,
 	int32 GutterSize,
-	UE::Geometry::TImageBuilder<FVector3f>& TextureOut,
+	UE::Geometry::FIndex4i BakeAttributes,
+	const FTextureAttributeSettings& AttributeSettings,
+	UE::Geometry::TImageBuilder<FVector4f>& TextureOut,
 	bool bOnlyOddMaterials = true,
 	TArrayView<int32> WhichMaterials = TArrayView<int32>()
 );
