Update FGeometryCollection's TransformToGeometryIndex mapping to INDEX_NONE when removing geometry elements

#rnx

diff --git a/Engine/Plugins/Experimental/PlanarCutPlugin/Source/PlanarCut/Private/GeometryMeshConversion.cpp b/Engine/Plugins/Experimental/PlanarCutPlugin/Source/PlanarCut/Private/GeometryMeshConversion.cpp
--- a/Engine/Plugins/Experimental/PlanarCutPlugin/Source/PlanarCut/Private/GeometryMeshConversion.cpp
+++ b/Engine/Plugins/Experimental/PlanarCutPlugin/Source/PlanarCut/Private/GeometryMeshConversion.cpp
@@ -1708,7 +1708,6 @@ int32 FDynamicMeshCollection::CutWithMultiplePlanes(
 
 		// remove old parent geometry
 		GeometryForRemoval.Add(GeometryIdx);
-		Collection->TransformToGeometryIndex[MeshData.TransformIndex] = INDEX_NONE;
 
 		// add newly created geometry as children
 		int32 SubPartIdx = 0;
@@ -1903,8 +1902,6 @@ int32 FDynamicMeshCollection::CutWithCellMeshes(const FInternalSurfaceMaterials&
 			
 			// remove old geom
 			GeometryForRemoval.Add(GeometryIdx);
-			Collection->TransformToGeometryIndex[Surface.TransformIndex] = INDEX_NONE;
-
 		}
 	}
 
diff --git a/Engine/Source/Runtime/Experimental/Chaos/Private/GeometryCollection/GeometryCollection.cpp b/Engine/Source/Runtime/Experimental/Chaos/Private/GeometryCollection/GeometryCollection.cpp
--- a/Engine/Source/Runtime/Experimental/Chaos/Private/GeometryCollection/GeometryCollection.cpp
+++ b/Engine/Source/Runtime/Experimental/Chaos/Private/GeometryCollection/GeometryCollection.cpp
@@ -623,8 +623,14 @@ void FGeometryCollection::RemoveGeometryElements(const TArray<int32>& SortedGeom
 			int32 GeoIndex = SortedGeometryIndicesToDelete[DeleteIdx];
 			for (int Idx = 0; Idx < TransformToGeometryIndex.Num(); Idx++)
 			{
-				if (TransformToGeometryIndex[Idx] >= GeoIndex)
+				if (TransformToGeometryIndex[Idx] > GeoIndex)
+				{
 					TransformToGeometryIndex[Idx]--;
+				}
+				else if (TransformToGeometryIndex[Idx] == GeoIndex)
+				{
+					TransformToGeometryIndex[Idx] = INDEX_NONE;
+				}
 			}
 
 		}
