Add a smoothed curvature attribute option to the fracture AutoUV tool (replaces less useful Z percent attribute)

diff --git a/Engine/Plugins/Experimental/ChaosEditor/Source/FractureEditor/Private/FractureToolUV.cpp b/Engine/Plugins/Experimental/ChaosEditor/Source/FractureEditor/Private/FractureToolUV.cpp
--- a/Engine/Plugins/Experimental/ChaosEditor/Source/FractureEditor/Private/FractureToolUV.cpp
+++ b/Engine/Plugins/Experimental/ChaosEditor/Source/FractureEditor/Private/FractureToolUV.cpp
@@ -225,14 +225,18 @@ int32 UFractureToolAutoUV::ExecuteFracture(const FFractureToolContext& FractureC
 		FIndex4i Attributes(
 			AutoUVSettings->bDistToOuter ? (int32)EBakeAttributes::DistanceToExternal : 0, 
 			AutoUVSettings->bAmbientOcclusion ? (int32)EBakeAttributes::AmbientOcclusion : 0, 
-			AutoUVSettings->bZNormal ? (int32)EBakeAttributes::NormalZ : 0,
-			AutoUVSettings->bZPosition ? (int32)EBakeAttributes::PositionZ : 0
+			AutoUVSettings->bSmoothedCurvature ? (int32)EBakeAttributes::Curvature : 0,
+			AutoUVSettings->bZNormal ? (int32)EBakeAttributes::NormalZ : 0
 		);
 		UE::PlanarCut::FTextureAttributeSettings AttribSettings;
 		AttribSettings.ToExternal_MaxDistance = AutoUVSettings->MaxDistance;
 		AttribSettings.AO_Rays = AutoUVSettings->OcclusionRays;
-		AttribSettings.bAO_Blur = AutoUVSettings->bGaussianBlur;
-		AttribSettings.AO_BlurRadius = AutoUVSettings->BlurRadius;
+		AttribSettings.AO_BlurRadius = AutoUVSettings->OcclusionBlurRadius;
+		AttribSettings.Curvature_BlurRadius = AutoUVSettings->CurvatureBlurRadius;
+		AttribSettings.Curvature_SmoothingSteps = AutoUVSettings->SmoothingIterations;
+		AttribSettings.Curvature_VoxelRes = AutoUVSettings->VoxelResolution;
+		AttribSettings.Curvature_ThicknessFactor = AutoUVSettings->ThicknessFactor;
+		AttribSettings.Curvature_MaxValue = AutoUVSettings->MaxCurvature;
 		AttribSettings.bNormalZ_TakeAbs = AutoUVSettings->bUseAbsoluteValue;
 		UE::PlanarCut::TextureInternalSurfaces(Collection, FMath::CeilToInt(AutoUVSettings->GutterSize), Attributes, AttribSettings, ImageBuilder);
 
diff --git a/Engine/Plugins/Experimental/ChaosEditor/Source/FractureEditor/Private/FractureToolUV.h b/Engine/Plugins/Experimental/ChaosEditor/Source/FractureEditor/Private/FractureToolUV.h
--- a/Engine/Plugins/Experimental/ChaosEditor/Source/FractureEditor/Private/FractureToolUV.h
+++ b/Engine/Plugins/Experimental/ChaosEditor/Source/FractureEditor/Private/FractureToolUV.h
@@ -60,13 +60,16 @@ public:
 	UPROPERTY(EditAnywhere, Category = AttributesToBake)
 	bool bAmbientOcclusion = true;
 
-	/** Bake the Z-component of the normal to a texture channel (blue) */
+	/**
+	 * Bake a smoothed curvature metric to a texture channel (blue)
+	 * Specifically, this is the mean curvature of a smoothed copy of each fractured piece, baked back to the respective fracture piece.
+	 */
 	UPROPERTY(EditAnywhere, Category = AttributesToBake)
-	bool bZNormal = true;
+	bool bSmoothedCurvature = true;
 
-	/** Bake the Z-component of the position (relative to the overall bounding box) to a texture channel (alpha) */
+	/** Bake the Z-component of the normal to a texture channel (alpha) */
 	UPROPERTY(EditAnywhere, Category = AttributesToBake)
-	bool bZPosition = false;
+	bool bZNormal = true;
 
 	/** Max distance to search for the outer mesh surface */
 	UPROPERTY(EditAnywhere, Category = DistToOuterSettings, meta = (EditCondition = "bDistToOuter", EditConditionHides, UIMin = "1", UIMax = "100", ClampMin = ".01", ClampMax = "1000"))
@@ -76,13 +79,29 @@ public:
 	UPROPERTY(EditAnywhere, Category = AmbientOcclusionSettings, meta = (EditCondition = "bAmbientOcclusion", EditConditionHides, UIMin = "1", UIMax = "1024", ClampMin = "0", ClampMax = "50000"))
 	int OcclusionRays = 16;
 
-	/** Whether or not to apply Gaussian Blur to computed AO Map (recommended) */
-	UPROPERTY(EditAnywhere, Category = AmbientOcclusionSettings, meta = (EditCondition = "bAmbientOcclusion", EditConditionHides))
-	bool bGaussianBlur = true;
+	/** Pixel Radius of Gaussian Blur Kernel applied to AO map (0 will apply no blur) */
+	UPROPERTY(EditAnywhere, Category = AmbientOcclusionSettings, meta = (EditCondition = "bAmbientOcclusion", EditConditionHides, UIMin = "0", UIMax = "10.0", ClampMin = "0", ClampMax = "100.0"))
+	double OcclusionBlurRadius = 2.25;
+
+	/** Pixel Radius of Gaussian Blur Kernel applied to Curvature map (0 will apply no blur) */
+	UPROPERTY(EditAnywhere, Category = SmoothedCurvatureSettings, meta = (EditCondition = "bSmoothedCurvature", EditConditionHides, UIMin = "0", UIMax = "10.0", ClampMin = "0", ClampMax = "100.0"))
+	double CurvatureBlurRadius = 2.25;
+
+	/** Voxel resolution of smoothed shape representation */
+	UPROPERTY(EditAnywhere, Category = SmoothedCurvatureSettings, meta = (EditCondition = "bSmoothedCurvature", EditConditionHides, UIMin = "8", UIMax = "512", ClampMin = "4", ClampMax = "1024"))
+	int VoxelResolution = 128;
+
+	/** Amount of smoothing iterations to apply before computing curvature */
+	UPROPERTY(EditAnywhere, Category = SmoothedCurvatureSettings, meta = (EditCondition = "bSmoothedCurvature", EditConditionHides, UIMin = "2", UIMax = "100", ClampMin = "2", ClampMax = "1000"))
+	int SmoothingIterations = 10;
+
+	/** Distance to search for correspondence between fractured shape and smoothed shape, as factor of voxel size */
+	UPROPERTY(EditAnywhere, Category = SmoothedCurvatureSettings, meta = (EditCondition = "bSmoothedCurvature", EditConditionHides, UIMin = "2", UIMax = "10.0", ClampMin = "1", ClampMax = "100.0"))
+	double ThicknessFactor = 4;
 
-	/** Pixel Radius of Gaussian Blur Kernel */
-	UPROPERTY(EditAnywhere, Category = AmbientOcclusionSettings, meta = (EditCondition = "bAmbientOcclusion && bGaussianBlur", EditConditionHides, UIMin = "0", UIMax = "10.0", ClampMin = "0", ClampMax = "100.0"))
-	double BlurRadius = 2.25;
+	/** Curvatures in the range [-MaxCurvature, MaxCurvature] will be mapped from [0,1]. Values outside that range will be clamped. */
+	UPROPERTY(EditAnywhere, Category = SmoothedCurvatureSettings, meta = (EditCondition = "bSmoothedCurvature", EditConditionHides, UIMin = ".01", UIMax = "1", ClampMin = ".0001", ClampMax = "10"))
+	double MaxCurvature = .1;
 
 	/** Whether to use the absolute value of the Z-Component of the normal */
 	UPROPERTY(EditAnywhere, Category = ZNormalSettings, meta = (EditCondition = "bZNormal", EditConditionHides))
diff --git a/Engine/Plugins/Experimental/GeometryProcessing/Source/DynamicMesh/Private/Sampling/MeshCurvatureMapBaker.cpp b/Engine/Plugins/Experimental/GeometryProcessing/Source/DynamicMesh/Private/Sampling/MeshCurvatureMapBaker.cpp
--- a/Engine/Plugins/Experimental/GeometryProcessing/Source/DynamicMesh/Private/Sampling/MeshCurvatureMapBaker.cpp
+++ b/Engine/Plugins/Experimental/GeometryProcessing/Source/DynamicMesh/Private/Sampling/MeshCurvatureMapBaker.cpp
@@ -34,6 +34,7 @@ void FMeshCurvatureMapBaker::Bake()
 
 	ResultBuilder = MakeUnique<TImageBuilder<FVector3f>>();
 	ResultBuilder->SetDimensions(BakeCache->GetDimensions());
+	ResultBuilder->Clear(FVector3f::Zero());
 
 
 	Bake_Single();
@@ -119,6 +120,10 @@ void FMeshCurvatureMapBaker::Bake_Single()
 	}
 
 	double ClampMax = RangeScale * (CurvatureStats.Mean + CurvatureStats.StandardDeviation);
+	if (bOverrideCurvatureRange)
+	{
+		ClampMax = RangeScale * OverrideRangeMax;
+	}
 	double MinClamp = MinRangeScale * ClampMax;
 	FInterval1d ClampRange(MinClamp, ClampMax);
 
diff --git a/Engine/Plugins/Experimental/GeometryProcessing/Source/DynamicMesh/Private/Sampling/MeshImageBakingCache.cpp b/Engine/Plugins/Experimental/GeometryProcessing/Source/DynamicMesh/Private/Sampling/MeshImageBakingCache.cpp
--- a/Engine/Plugins/Experimental/GeometryProcessing/Source/DynamicMesh/Private/Sampling/MeshImageBakingCache.cpp
+++ b/Engine/Plugins/Experimental/GeometryProcessing/Source/DynamicMesh/Private/Sampling/MeshImageBakingCache.cpp
@@ -35,6 +35,16 @@ void FMeshImageBakingCache::SetDimensions(FImageDimensions DimensionsIn)
 }
 
 
+void FMeshImageBakingCache::SetGutterSize(int32 GutterSizeIn)
+{
+	if (GutterSize != GutterSizeIn)
+	{
+		GutterSize = GutterSizeIn;
+		InvalidateOccupancy();
+	}
+}
+
+
 void FMeshImageBakingCache::SetUVLayer(int32 UVLayerIn)
 {
 	if (UVLayer != UVLayerIn)
@@ -271,6 +281,7 @@ bool FMeshImageBakingCache::ValidateCache()
 
 		// calculate occupancy map
 		OccupancyMap = MakeUnique<FImageOccupancyMap>();
+		OccupancyMap->GutterSize = GutterSize;
 		OccupancyMap->Initialize(Dimensions);
 		OccupancyMap->ComputeFromUVSpaceMesh(FlatMesh, [&](int32 TriangleID) { return FlatMesh.GetTriangleGroup(TriangleID); });
 
diff --git a/Engine/Plugins/Experimental/GeometryProcessing/Source/DynamicMesh/Public/Sampling/MeshCurvatureMapBaker.h b/Engine/Plugins/Experimental/GeometryProcessing/Source/DynamicMesh/Public/Sampling/MeshCurvatureMapBaker.h
--- a/Engine/Plugins/Experimental/GeometryProcessing/Source/DynamicMesh/Public/Sampling/MeshCurvatureMapBaker.h
+++ b/Engine/Plugins/Experimental/GeometryProcessing/Source/DynamicMesh/Public/Sampling/MeshCurvatureMapBaker.h
@@ -50,8 +50,12 @@ public:
 	double RangeScale = 1.0;
 	double MinRangeScale = 0.0;
 
-	double BlurRadius = 0.0;
+	// allows override of the max curvature; if false, range is set based on [-(avg+stddev), avg+stddev]
+	bool bOverrideCurvatureRange = false;
+	double OverrideRangeMax = .1;
+
 
+	double BlurRadius = 0.0;
 
 	//
 	// Required input data, can be provided, will be computed otherwise
diff --git a/Engine/Plugins/Experimental/GeometryProcessing/Source/DynamicMesh/Public/Sampling/MeshImageBakingCache.h b/Engine/Plugins/Experimental/GeometryProcessing/Source/DynamicMesh/Public/Sampling/MeshImageBakingCache.h
--- a/Engine/Plugins/Experimental/GeometryProcessing/Source/DynamicMesh/Public/Sampling/MeshImageBakingCache.h
+++ b/Engine/Plugins/Experimental/GeometryProcessing/Source/DynamicMesh/Public/Sampling/MeshImageBakingCache.h
@@ -43,6 +43,7 @@ public:
 	void SetUVLayer(int32 UVLayer);
 	void SetThickness(double Thickness);
 	void SetCorrespondenceStrategy(ECorrespondenceStrategy Strategy);
+	void SetGutterSize(int32 GutterSize);
 
 	FImageDimensions GetDimensions() const { return Dimensions; }
 	int32 GetUVLayer() const { return UVLayer; }
@@ -104,6 +105,7 @@ protected:
 	int32 UVLayer = 0;
 	double Thickness = 3.0;
 	ECorrespondenceStrategy CorrespondenceStrategy = ECorrespondenceStrategy::RaycastStandard;
+	int32 GutterSize = 4;
 
 	TDenseGrid2<FCorrespondenceSample> SampleMap;
 	bool bSamplesValid = false;
diff --git a/Engine/Plugins/Experimental/PlanarCutPlugin/Source/PlanarCut/Private/FractureAutoUV.cpp b/Engine/Plugins/Experimental/PlanarCutPlugin/Source/PlanarCut/Private/FractureAutoUV.cpp
--- a/Engine/Plugins/Experimental/PlanarCutPlugin/Source/PlanarCut/Private/FractureAutoUV.cpp
+++ b/Engine/Plugins/Experimental/PlanarCutPlugin/Source/PlanarCut/Private/FractureAutoUV.cpp
@@ -10,10 +10,14 @@
 
 #include "DynamicMesh3.h"
 #include "MeshNormals.h"
+#include "MeshWeights.h"
 #include "Parameterization/DynamicMeshUVEditor.h"
 #include "Sampling/MeshImageBakingCache.h"
 
+#include "Implicit/Solidify.h"
+
 #include "Parameterization/UVPacking.h"
+#include "Sampling/MeshCurvatureMapBaker.h"
 #include "Sampling/MeshOcclusionMapBaker.h"
 #include "DisjointSet.h"
 
@@ -497,15 +501,17 @@ void TextureInternalSurfaces(
 
 	int DistanceToExternalIdx = BakeAttributes.IndexOf((int)EBakeAttributes::DistanceToExternal);
 	int AmbientIdx = BakeAttributes.IndexOf((int)EBakeAttributes::AmbientOcclusion);
+	int CurvatureIdx = BakeAttributes.IndexOf((int)EBakeAttributes::Curvature);
 	int NormalZIdx = BakeAttributes.IndexOf((int)EBakeAttributes::NormalZ);
 	int PosZIdx = BakeAttributes.IndexOf((int)EBakeAttributes::PositionZ);
+	
 	if (OutsideSpatial.GetBoundingBox().Depth() == 0)
 	{
 		PosZIdx = -1; // if everything has same Z, can just leave PosZ values as default value
 	}
 
 
-	bool bNeedsDynamicMeshes = AmbientIdx > -1;
+	bool bNeedsDynamicMeshes = AmbientIdx > -1 || CurvatureIdx > -1;
 
 	TArray<int32> TransformIndices;
 	if (bNeedsDynamicMeshes)
@@ -544,6 +550,104 @@ void TextureInternalSurfaces(
 			}
 		}
 	}
+
+	TArray64<float> CurvatureValues; // buffer to fill with curvature values
+	if (CurvatureIdx > -1)
+	{
+		int UseVoxRes = FMath::Max(8, AttributeSettings.Curvature_VoxelRes);
+		FImageDimensions Dimensions = TextureOut.GetDimensions();
+		int64 TexelsNum = Dimensions.Num();
+		CurvatureValues.SetNumZeroed(TexelsNum);
+		for (int MeshIdx = 0; MeshIdx < CollectionMeshes.Meshes.Num(); MeshIdx++)
+		{
+
+			int32 TransformIdx = CollectionMeshes.Meshes[MeshIdx].TransformIndex;
+			FDynamicMesh3& Mesh = CollectionMeshes.Meshes[MeshIdx].AugMesh;
+			FDynamicMeshUVOverlay* UV = Mesh.Attributes()->PrimaryUV();
+
+			FDynamicMeshAABBTree3 Spatial(&Mesh);
+			TFastWindingTree<FDynamicMesh3> FastWinding(&Spatial);
+			double ExtendBounds = Spatial.GetBoundingBox().MaxDim() * .01;
+
+			TImplicitSolidify<FDynamicMesh3> Solidify(&Mesh, &Spatial, &FastWinding);
+			Solidify.SetCellSizeAndExtendBounds(Spatial.GetBoundingBox(), ExtendBounds, UseVoxRes);
+			Solidify.WindingThreshold = AttributeSettings.Curvature_Winding;
+			Solidify.SurfaceSearchSteps = 3;
+			Solidify.bSolidAtBoundaries = true;
+			Solidify.ExtendBounds = ExtendBounds;
+
+			FDynamicMesh3 SolidMesh(&Solidify.Generate());
+
+			// Smooth mesh (there is a more general version in MeshModelingToolset's IterativeSmoothingOp.cpp, but don't want to depend on MeshModelingToolset just for that)
+			// (note it is extra simple b/c we shouldn't have boundary vertices)
+			auto IterativeSmoothMesh = [](FDynamicMesh3& Mesh, int Iters, double SmoothAlpha)
+			{
+				TArray<FVector3d> SmoothedBuffer, PositionBuffer;
+				int NV = Mesh.MaxVertexID();
+				SmoothedBuffer.SetNumZeroed(NV);
+				PositionBuffer.SetNumZeroed(NV);
+				for (int VID : Mesh.VertexIndicesItr())
+				{
+					PositionBuffer[VID] = Mesh.GetVertex(VID);
+				}
+				for (int32 k = 0; k < Iters; ++k)
+				{
+					// calculate smoothed positions of interior vertices
+					ParallelFor(NV, [&Mesh, &SmoothedBuffer, &PositionBuffer, SmoothAlpha](int32 vid)
+						{
+							if (Mesh.IsReferencedVertex(vid) == false)
+							{
+								return;
+							}
+
+							FVector3d Centroid = FMeshWeights::UniformCentroid(Mesh, vid, [&](int32 nbrvid) { return PositionBuffer[nbrvid]; });
+							SmoothedBuffer[vid] = Lerp(PositionBuffer[vid], Centroid, SmoothAlpha);
+						});
+
+					for (int32 vid = 0; vid < NV; ++vid)
+					{
+						PositionBuffer[vid] = SmoothedBuffer[vid];
+					}
+				}
+				for (int VID : Mesh.VertexIndicesItr())
+				{
+					Mesh.SetVertex(VID, PositionBuffer[VID]);
+				}
+			};
+			IterativeSmoothMesh(SolidMesh, AttributeSettings.Curvature_SmoothingSteps, AttributeSettings.Curvature_SmoothingPerStep);
+
+			FDynamicMeshAABBTree3 SolidSpatial(&SolidMesh);
+
+			FMeshImageBakingCache BakeCache;
+			BakeCache.SetGutterSize(GutterSize);
+			BakeCache.SetDetailMesh(&SolidMesh, &SolidSpatial);
+			BakeCache.SetBakeTargetMesh(&Mesh);
+			BakeCache.SetDimensions(Dimensions);
+			BakeCache.SetUVLayer(0);
+			// set distance to search for correspondences based on the voxel size
+			double Thickness = AttributeSettings.Curvature_ThicknessFactor * Spatial.GetBoundingBox().MaxDim() / (float)UseVoxRes;
+			BakeCache.SetThickness(Thickness);
+			BakeCache.ValidateCache();
+
+			FMeshCurvatureMapBaker CurvatureBaker;
+			CurvatureBaker.bOverrideCurvatureRange = true;
+			CurvatureBaker.OverrideRangeMax = AttributeSettings.Curvature_MaxValue;
+			CurvatureBaker.UseColorMode = FMeshCurvatureMapBaker::EColorMode::BlackGrayWhite;
+			CurvatureBaker.SetCache(&BakeCache);
+
+			CurvatureBaker.BlurRadius = AttributeSettings.bCurvature_Blur ? AttributeSettings.Curvature_BlurRadius : 0.0;
+
+			CurvatureBaker.Bake();
+
+			const TUniquePtr<TImageBuilder<FVector3f>>& Result = CurvatureBaker.GetResult();
+			// copy scalar curvature values out to the accumulated curvature buffer
+			for (int64 LinearIdx = 0; LinearIdx < TexelsNum; LinearIdx++)
+			{
+				float& CurvatureOut = CurvatureValues[LinearIdx];
+				CurvatureOut = FMathf::Max(CurvatureOut, Result->GetPixel(LinearIdx).X);
+			}
+		}
+	}
 	
 	TArray64<float> AmbientValues; // buffer to fill with ambient occlusion values
 	if (AmbientIdx > -1)
@@ -560,6 +664,7 @@ void TextureInternalSurfaces(
 			FDynamicMeshAABBTree3 Spatial(&Mesh);
 
 			FMeshImageBakingCache BakeCache;
+			BakeCache.SetGutterSize(GutterSize);
 			BakeCache.SetDetailMesh(&Mesh, &Spatial);
 			BakeCache.SetBakeTargetMesh(&Mesh);
 			BakeCache.SetDimensions(Dimensions);
@@ -655,17 +760,22 @@ void TextureInternalSurfaces(
 		}
 	});
 
-	// Copy the ambient occlusion values to the appropriate channel
-	if (AmbientIdx > -1)
+	auto CopyValuesToChannel = [&TextureOut, &OccupancyMap](int ChannelIdx, TArray64<float>& Values)
 	{
-		int64 TexelsNum = OccupancyMap.Dimensions.Num();
-		for (int64 LinearIdx = 0; LinearIdx < TexelsNum; LinearIdx++)
+		if (ChannelIdx > -1)
 		{
-			FVector4f Pixel = TextureOut.GetPixel(LinearIdx);
-			Pixel[AmbientIdx] = AmbientValues[LinearIdx];
-			TextureOut.SetPixel(LinearIdx, Pixel);
+			int64 TexelsNum = OccupancyMap.Dimensions.Num();
+			for (int64 LinearIdx = 0; LinearIdx < TexelsNum; LinearIdx++)
+			{
+				FVector4f Pixel = TextureOut.GetPixel(LinearIdx);
+				Pixel[ChannelIdx] = Values[LinearIdx];
+				TextureOut.SetPixel(LinearIdx, Pixel);
+			}
 		}
-	}
+	};
+
+	CopyValuesToChannel(AmbientIdx, AmbientValues);
+	CopyValuesToChannel(CurvatureIdx, CurvatureValues);
 
 	for (TTuple<int64, int64>& GutterToInside : OccupancyMap.GutterTexels)
 	{
diff --git a/Engine/Plugins/Experimental/PlanarCutPlugin/Source/PlanarCut/Public/FractureAutoUV.h b/Engine/Plugins/Experimental/PlanarCutPlugin/Source/PlanarCut/Public/FractureAutoUV.h
--- a/Engine/Plugins/Experimental/PlanarCutPlugin/Source/PlanarCut/Public/FractureAutoUV.h
+++ b/Engine/Plugins/Experimental/PlanarCutPlugin/Source/PlanarCut/Public/FractureAutoUV.h
@@ -34,7 +34,7 @@ enum class EBakeAttributes : int32
 	None,
 	DistanceToExternal,
 	AmbientOcclusion,
-	//Curvature, // TODO: add a smoothed curvature attribute
+	Curvature,
 	NormalZ,
 	PositionZ
 };
@@ -47,6 +47,15 @@ struct FTextureAttributeSettings
 	bool bAO_Blur = true;
 	double AO_BlurRadius = 2.5;
 	double AO_MaxDistance = 0.0; // 0.0 is interpreted as TNumericLimits<double>::Max()
+	int Curvature_VoxelRes = 128;
+	double Curvature_Winding = .5;
+	int Curvature_SmoothingSteps = 10;
+	double Curvature_SmoothingPerStep = .8;
+	bool bCurvature_Blur = true;
+	double Curvature_BlurRadius = 2.5;
+	double Curvature_ThicknessFactor = 3.0; // distance to search for mesh correspondence, as a factor of voxel size
+	double Curvature_MaxValue = .1; // curvatures above this value will be clamped
+	
 	bool bNormalZ_TakeAbs = true;
 };
 
