fix AutoUV crash -- smoothing refactor was applied to wrong mesh

#rnx

diff --git a/Engine/Plugins/Experimental/PlanarCutPlugin/Source/PlanarCut/Private/FractureAutoUV.cpp b/Engine/Plugins/Experimental/PlanarCutPlugin/Source/PlanarCut/Private/FractureAutoUV.cpp
--- a/Engine/Plugins/Experimental/PlanarCutPlugin/Source/PlanarCut/Private/FractureAutoUV.cpp
+++ b/Engine/Plugins/Experimental/PlanarCutPlugin/Source/PlanarCut/Private/FractureAutoUV.cpp
@@ -582,16 +582,16 @@ void TextureInternalSurfaces(
 			// Smooth mesh
 			double SmoothAlpha = AttributeSettings.Curvature_SmoothingPerStep;
 			TArray<FVector3d> PositionBuffer;
-			PositionBuffer.SetNumUninitialized(Mesh.MaxVertexID());
-			for (int VID : Mesh.VertexIndicesItr())
+			PositionBuffer.SetNumUninitialized(SolidMesh.MaxVertexID());
+			for (int VID : SolidMesh.VertexIndicesItr())
 			{
-				PositionBuffer[VID] = Mesh.GetVertex(VID);
+				PositionBuffer[VID] = SolidMesh.GetVertex(VID);
 			}
 			UE::MeshDeformation::ComputeSmoothing_Forward(true, false, SolidMesh, [SmoothAlpha](int VID, bool bIsBoundary) { return SmoothAlpha; },
 				AttributeSettings.Curvature_SmoothingSteps, PositionBuffer);
-			for (int VID : Mesh.VertexIndicesItr())
+			for (int VID : SolidMesh.VertexIndicesItr())
 			{
-				Mesh.SetVertex(VID, PositionBuffer[VID]);
+				SolidMesh.SetVertex(VID, PositionBuffer[VID]);
 			}
 
 			FDynamicMeshAABBTree3 SolidSpatial(&SolidMesh);
