fix fracture plane cuts having incorrect normals
#rnx

diff --git a/Engine/Plugins/Experimental/PlanarCutPlugin/Source/PlanarCut/Private/GeometryMeshConversion.cpp b/Engine/Plugins/Experimental/PlanarCutPlugin/Source/PlanarCut/Private/GeometryMeshConversion.cpp
--- a/Engine/Plugins/Experimental/PlanarCutPlugin/Source/PlanarCut/Private/GeometryMeshConversion.cpp
+++ b/Engine/Plugins/Experimental/PlanarCutPlugin/Source/PlanarCut/Private/GeometryMeshConversion.cpp
@@ -1168,7 +1168,7 @@ void FCellMeshes::CreateMeshesForSinglePlane(const FPlanarCells& Cells, const FA
 	PlaneVertInfo.bHaveUV = true;
 	PlaneVertInfo.bHaveN = true;
 	PlaneVertInfo.Color = FVector3f(1, 1, 1);
-	PlaneVertInfo.Normal = FVector3f(Plane.GetNormal());
+	PlaneVertInfo.Normal = -FVector3f(Plane.GetNormal());
 
 	for (int CornerIdx = 0; CornerIdx < 4; CornerIdx++)
 	{
