GeometryCollection
-- Support for Nanite to GeometryCollection conversion
-- Support for multiple UV sets

diff --git a/Engine/Plugins/Experimental/ChaosEditor/Source/FractureEditor/Private/FractureToolGenerators.cpp b/Engine/Plugins/Experimental/ChaosEditor/Source/FractureEditor/Private/FractureToolGenerators.cpp
--- a/Engine/Plugins/Experimental/ChaosEditor/Source/FractureEditor/Private/FractureToolGenerators.cpp
+++ b/Engine/Plugins/Experimental/ChaosEditor/Source/FractureEditor/Private/FractureToolGenerators.cpp
@@ -204,7 +204,7 @@ AGeometryCollectionActor* UFractureToolGenerateAsset::ConvertActorsToGeometryCol
 					decltype(FGeometryCollectionSource::SourceMaterial) SourceMaterials(StaticMeshComponent->GetMaterials());
 					FracturedGeometryCollection->GeometrySource.Add({ SourceSoftObjectPath, ComponentTransform, SourceMaterials });
 
-					FGeometryCollectionConversion::AppendStaticMesh(ComponentStaticMesh, SourceMaterials, ComponentTransform, FracturedGeometryCollection, false);
+					FGeometryCollectionConversion::AppendStaticMesh(ComponentStaticMesh, SourceMaterials, ComponentTransform, FracturedGeometryCollection, true);
 				}
 			}
 		}
diff --git a/Engine/Plugins/Experimental/GeometryCollectionPlugin/Source/GeometryCollectionEditor/Private/GeometryCollection/GeometryCollectionConversion.cpp b/Engine/Plugins/Experimental/GeometryCollectionPlugin/Source/GeometryCollectionEditor/Private/GeometryCollection/GeometryCollectionConversion.cpp
--- a/Engine/Plugins/Experimental/GeometryCollectionPlugin/Source/GeometryCollectionEditor/Private/GeometryCollection/GeometryCollectionConversion.cpp
+++ b/Engine/Plugins/Experimental/GeometryCollectionPlugin/Source/GeometryCollectionEditor/Private/GeometryCollection/GeometryCollectionConversion.cpp
@@ -33,13 +33,24 @@ struct FUniqueVertex
 {
 	FVector Normal;
 	FVector Tangent;
-	FVector2D UV;
+	TArray<FVector2D> UVs;
 
 	bool operator==(const FUniqueVertex& Other) const
 	{
-		return ((this->Normal == Other.Normal) &&
-			(this->Tangent == Other.Tangent) &&
-			(this->UV == Other.UV));
+		if (this->UVs.Num() != Other.UVs.Num())
+		{
+			return false;
+		}
+		
+		bool bEquality = true;
+		bEquality &= (this->Normal == Other.Normal);
+		bEquality &= (this->Tangent == Other.Tangent);
+		for (int32 UVLayerIdx = 0; UVLayerIdx < UVs.Num(); ++UVLayerIdx)
+		{
+			bEquality &= (this->UVs[UVLayerIdx] == Other.UVs[UVLayerIdx]);
+		}
+		
+		return bEquality;
 	}
 };
 
@@ -47,8 +58,11 @@ FORCEINLINE uint32 GetTypeHash(const FUniqueVertex& UniqueVertex)
 {
 	uint32 VertexHash = GetTypeHash(UniqueVertex.Normal);
 	VertexHash = HashCombine(VertexHash, GetTypeHash(UniqueVertex.Tangent));
-	VertexHash = HashCombine(VertexHash, GetTypeHash(UniqueVertex.UV));
-
+	for (int32 UVLayerIdx = 0; UVLayerIdx < UniqueVertex.UVs.Num(); ++UVLayerIdx)
+	{
+		VertexHash = HashCombine(VertexHash, GetTypeHash(UniqueVertex.UVs[UVLayerIdx]));
+	}
+	
 	return VertexHash;
 }
 
@@ -89,15 +103,20 @@ void FGeometryCollectionConversion::AppendStaticMesh(const UStaticMesh* StaticMe
 		TArrayView<const FVector4> SourceColor = Attributes.GetVertexInstanceColors().GetRawArray();
 
 		TVertexInstanceAttributesConstRef<FVector2D> InstanceUVs = Attributes.GetVertexInstanceUVs();
-		TArrayView<const FVector2D> SourceUV = InstanceUVs.GetRawArray();
-		const int NumUVLayers = InstanceUVs.GetNumChannels();
+		const int32 NumUVLayers = InstanceUVs.GetNumChannels();
+		TArray<TArrayView<const FVector2D>> SourceUVArrays;
+		SourceUVArrays.SetNum(NumUVLayers);
+		for (int32 UVLayerIdx = 0; UVLayerIdx < NumUVLayers; ++UVLayerIdx)
+		{
+			SourceUVArrays[UVLayerIdx] = InstanceUVs.GetRawArray(UVLayerIdx);
+		}
 		
 		// target vertex information
 		TManagedArray<FVector3f>& TargetVertex = GeometryCollection->Vertex;
 		TManagedArray<FVector3f>& TargetTangentU = GeometryCollection->TangentU;
 		TManagedArray<FVector3f>& TargetTangentV = GeometryCollection->TangentV;
 		TManagedArray<FVector3f>& TargetNormal = GeometryCollection->Normal;
-		TManagedArray<FVector2D>& TargetUV = GeometryCollection->UV;
+		TManagedArray<TArray<FVector2D>>& TargetUVs = GeometryCollection->UVs;
 		TManagedArray<FLinearColor>& TargetColor = GeometryCollection->Color;
 		TManagedArray<int32>& TargetBoneMap = GeometryCollection->BoneMap;
 		TManagedArray<FLinearColor>& TargetBoneColor = GeometryCollection->BoneColor;
@@ -114,15 +133,21 @@ void FGeometryCollectionConversion::AppendStaticMesh(const UStaticMesh* StaticMe
 		VertexInstanceToGeometryCollectionVertex.Reserve(Attributes.GetVertexInstanceNormals().GetNumElements());
 		
 		for (const FVertexID VertexIndex : MeshDescription->Vertices().GetElementIDs())
-		{
-			
+		{		
 			TArrayView<const FVertexInstanceID> ReferencingVertexInstances = MeshDescription->GetVertexVertexInstanceIDs(VertexIndex);
 
 			// Generate per instance hash of splittable attributes.
 			TMap<FUniqueVertex, TArray<FVertexInstanceID>> SplitVertices;
 			for (const FVertexInstanceID& InstanceID : ReferencingVertexInstances)
 			{
-				FUniqueVertex UniqueVertex{ SourceNormal[InstanceID], SourceTangent[InstanceID], SourceUV[InstanceID] };
+				TArray<FVector2D> SourceUVs;
+				SourceUVs.SetNum(NumUVLayers);
+				for (int32 UVLayerIdx = 0; UVLayerIdx < NumUVLayers; ++UVLayerIdx)
+				{
+					SourceUVs[UVLayerIdx] = SourceUVArrays[UVLayerIdx][InstanceID];
+				}
+				
+				FUniqueVertex UniqueVertex{ SourceNormal[InstanceID], SourceTangent[InstanceID], SourceUVs };
 				TArray<FVertexInstanceID>& SplitVertex = SplitVertices.FindOrAdd(UniqueVertex);
 				SplitVertex.Add(InstanceID);
 			}
@@ -142,8 +167,7 @@ void FGeometryCollectionConversion::AppendStaticMesh(const UStaticMesh* StaticMe
 				TargetTangentU[CurrentVertex] = SourceTangent[ExemplarInstanceID];
 				TargetTangentV[CurrentVertex] = SourceBinormalSign[ExemplarInstanceID] * FVector::CrossProduct(TargetNormal[CurrentVertex], TargetTangentU[CurrentVertex]);
 
-				// @todo : Support multiple UV's per vertex based on MAX_STATIC_TEXCOORDS
-				TargetUV[CurrentVertex] = SourceUV[ExemplarInstanceID];
+				TargetUVs[CurrentVertex] = SplitVertex.Key.UVs;
 
 				if (SourceColor.Num() > 0)
 				{
@@ -164,6 +188,27 @@ void FGeometryCollectionConversion::AppendStaticMesh(const UStaticMesh* StaticMe
 			}
 		}
 
+
+		// for each material, add a reference in our GeometryCollectionObject
+		const int32 MaterialStart = GeometryCollectionObject->Materials.Num();
+		const int32 NumMeshMaterials = Materials.Num();
+		GeometryCollectionObject->Materials.Reserve(MaterialStart + NumMeshMaterials);
+
+		for (int32 Index = 0; Index < NumMeshMaterials; ++Index)
+		{
+			UMaterialInterface* CurrMaterial = Materials[Index];
+
+			// Possible we have a null entry - replace with default
+			if (CurrMaterial == nullptr)
+			{
+				CurrMaterial = UMaterial::GetDefaultMaterial(MD_Surface);
+			}
+
+			// We add the material twice, once for interior and again for exterior.
+			GeometryCollectionObject->Materials.Add(CurrMaterial);
+			GeometryCollectionObject->Materials.Add(CurrMaterial);
+		}
+
 		// target triangle indices
 		TManagedArray<FIntVector>& TargetIndices = GeometryCollection->Indices;
 		TManagedArray<bool>& TargetVisible = GeometryCollection->Visible;
@@ -185,7 +230,9 @@ void FGeometryCollectionConversion::AppendStaticMesh(const UStaticMesh* StaticMe
 			);
 
 			TargetVisible[TargetIndex] = true;
-			TargetMaterialID[TargetIndex] = 0;
+
+			// Materials are ganged in pairs and we want the id to associate with the first of each pair.
+			TargetMaterialID[TargetIndex] = MaterialStart + (MeshDescription->GetTrianglePolygonGroup(TriangleIndex) * 2);
 
 			// Is this right?
 			TargetMaterialIndex[TargetIndex] = TargetIndex;
@@ -280,54 +327,6 @@ void FGeometryCollectionConversion::AppendStaticMesh(const UStaticMesh* StaticMe
 			}
 		}
 
-		// for each material, add a reference in our GeometryCollectionObject
-		const int32 MaterialStart = GeometryCollectionObject->Materials.Num();
-		const int32 NumMeshMaterials = Materials.Num();
-		GeometryCollectionObject->Materials.Reserve(MaterialStart + NumMeshMaterials);
-
-		for (int32 Index = 0; Index < NumMeshMaterials; ++Index)
-		{
-			UMaterialInterface* CurrMaterial = Materials[Index];
-
-			// Possible we have a null entry - replace with default
-			if (CurrMaterial == nullptr)
-			{
-				CurrMaterial = UMaterial::GetDefaultMaterial(MD_Surface);
-			}
-
-			// We add the material twice, once for interior and again for exterior.
-			GeometryCollectionObject->Materials.Add(CurrMaterial);
-			GeometryCollectionObject->Materials.Add(CurrMaterial);
-		}
-
-		TManagedArray<FGeometryCollectionSection>& Sections = GeometryCollection->Sections;
-
-		// We make sections that mirror what is in the static mesh.  Note that this isn't explicitly
-		// necessary since we reindex after all the meshes are added, but it is a good step to have
-		// optimal min/max vertex index right from the static mesh.  All we really need to do is
-		// assign material ids and rely on reindexing, in theory
-		for (const FStaticMeshSection& CurrSection : StaticMesh->GetRenderData()->LODResources[0].Sections)
-		{			
-			// create new section
-			int32 SectionIndex = GeometryCollection->AddElements(1, FGeometryCollection::MaterialGroup);
-
-			// Materials are doubled on ingestions: we multiply the incoming ID by 2 to skip internal materials.
-			Sections[SectionIndex].MaterialID = MaterialStart + (CurrSection.MaterialIndex * 2);
-
-			Sections[SectionIndex].FirstIndex = IndicesStart * 3 + CurrSection.FirstIndex;
-			Sections[SectionIndex].MinVertexIndex = VertexStart + CurrSection.MinVertexIndex;
-
-			Sections[SectionIndex].NumTriangles = CurrSection.NumTriangles;
-			Sections[SectionIndex].MaxVertexIndex = VertexStart + CurrSection.MaxVertexIndex;
-
-			// set the MaterialID for all of the faces
-			// note the divide by 3 - the GeometryCollection stores indices in tuples of 3 rather than in a flat array
-			for (int32 i = Sections[SectionIndex].FirstIndex / 3; i < Sections[SectionIndex].FirstIndex / 3 + Sections[SectionIndex].NumTriangles; ++i)
-			{
-				TargetMaterialID[i] = Sections[SectionIndex].MaterialID;
-			}
-		}
-
 		if (ReindexMaterials) {
 			GeometryCollection->ReindexMaterials();
 		}
@@ -371,7 +370,7 @@ void FGeometryCollectionConversion::AppendGeometryCollection(const UGeometryColl
 	const TManagedArray<FVector3f>& SourceTangentU = SourceGeometryCollectionPtr->TangentU;
 	const TManagedArray<FVector3f>& SourceTangentV = SourceGeometryCollectionPtr->TangentV;
 	const TManagedArray<FVector3f>& SourceNormal = SourceGeometryCollectionPtr->Normal;
-	const TManagedArray<FVector2D>& SourceUV = SourceGeometryCollectionPtr->UV;
+	const TManagedArray<TArray<FVector2D>>& SourceUVs = SourceGeometryCollectionPtr->UVs;
 	const TManagedArray<FLinearColor>& SourceColor = SourceGeometryCollectionPtr->Color;
 	const TManagedArray<int32>& SourceBoneMap = SourceGeometryCollectionPtr->BoneMap;
 
@@ -380,7 +379,7 @@ void FGeometryCollectionConversion::AppendGeometryCollection(const UGeometryColl
 	TManagedArray<FVector3f>& TargetTangentU = GeometryCollection->TangentU;
 	TManagedArray<FVector3f>& TargetTangentV = GeometryCollection->TangentV;
 	TManagedArray<FVector3f>& TargetNormal = GeometryCollection->Normal;
-	TManagedArray<FVector2D>& TargetUV = GeometryCollection->UV;
+	TManagedArray<TArray<FVector2D>>& TargetUVs = GeometryCollection->UVs;
 	TManagedArray<FLinearColor>& TargetColor = GeometryCollection->Color;
 	TManagedArray<int32>& TargetBoneMap = GeometryCollection->BoneMap;
 
@@ -393,7 +392,7 @@ void FGeometryCollectionConversion::AppendGeometryCollection(const UGeometryColl
 		TargetTangentU[VertexOffset] = SourceTangentU[VertexIndex];
 		TargetTangentV[VertexOffset] = SourceTangentV[VertexIndex];
 		TargetNormal[VertexOffset] = SourceNormal[VertexIndex];
-		TargetUV[VertexOffset] = SourceUV[VertexIndex];
+		TargetUVs[VertexOffset] = SourceUVs[VertexIndex];
 		TargetColor[VertexOffset] = SourceColor[VertexIndex];
 
 		TargetBoneMap[VertexOffset] = SourceBoneMap[VertexIndex] + TransformStart;
@@ -717,7 +716,7 @@ void FGeometryCollectionConversion::AppendSkeletalMesh(const USkeletalMesh* Skel
 						TManagedArray<FVector3f>& TangentU = GeometryCollection->TangentU;
 						TManagedArray<FVector3f>& TangentV = GeometryCollection->TangentV;
 						TManagedArray<FVector3f>& Normal = GeometryCollection->Normal;
-						TManagedArray<FVector2D>& UV = GeometryCollection->UV;
+						TManagedArray<TArray<FVector2D>>& UVs = GeometryCollection->UVs;
 						TManagedArray<FLinearColor>& Color = GeometryCollection->Color;
 						TManagedArray<int32>& BoneMap = GeometryCollection->BoneMap;
 						TManagedArray<FLinearColor>& BoneColor = GeometryCollection->BoneColor;
@@ -730,6 +729,7 @@ void FGeometryCollectionConversion::AppendSkeletalMesh(const USkeletalMesh* Skel
 						TManagedArray<int32>& SimulationType = GeometryCollection->SimulationType;
 						int InitialNumVertices = GeometryCollection->NumElements(FGeometryCollection::VerticesGroup);
 						int VertexBaseIndex = GeometryCollection->AddElements(VertexCount, FGeometryCollection::VerticesGroup);
+						const int32 NumUVLayers = VertexBuffers.StaticMeshVertexBuffer.GetNumTexCoords();
 						for (int32 VertexIndex = 0; VertexIndex < VertexCount; VertexIndex++)
 						{
 							int VertexOffset = VertexBaseIndex + VertexIndex;
@@ -745,8 +745,13 @@ void FGeometryCollectionConversion::AppendSkeletalMesh(const USkeletalMesh* Skel
 							TangentU[VertexOffset] = VertexBuffers.StaticMeshVertexBuffer.VertexTangentX(VertexIndex);
 							TangentV[VertexOffset] = VertexBuffers.StaticMeshVertexBuffer.VertexTangentY(VertexIndex);
 							Normal[VertexOffset] = VertexBuffers.StaticMeshVertexBuffer.VertexTangentZ(VertexIndex);
-							// @todo : Support multiple UV's per vertex based on MAX_STATIC_TEXCOORDS 
-							UV[VertexOffset] = VertexBuffers.StaticMeshVertexBuffer.GetVertexUV(VertexIndex, 0);
+							
+							UVs[VertexOffset].SetNum(NumUVLayers);
+							for (int32 UVLayerIdx = 0; UVLayerIdx < NumUVLayers; ++UVLayerIdx)
+							{
+								UVs[VertexOffset][UVLayerIdx] = VertexBuffers.StaticMeshVertexBuffer.GetVertexUV(VertexIndex, UVLayerIdx);
+							}
+							
 							if (VertexBuffers.ColorVertexBuffer.GetNumVertices() == VertexCount)
 								Color[VertexOffset] = VertexBuffers.ColorVertexBuffer.VertexColor(VertexIndex);
 							else
diff --git a/Engine/Plugins/Experimental/PlanarCutPlugin/Source/PlanarCut/Private/FractureAutoUV.cpp b/Engine/Plugins/Experimental/PlanarCutPlugin/Source/PlanarCut/Private/FractureAutoUV.cpp
--- a/Engine/Plugins/Experimental/PlanarCutPlugin/Source/PlanarCut/Private/FractureAutoUV.cpp
+++ b/Engine/Plugins/Experimental/PlanarCutPlugin/Source/PlanarCut/Private/FractureAutoUV.cpp
@@ -217,12 +217,12 @@ struct FGeomMesh : public UE::Geometry::FUVPacker::IUVMeshView
 
 	virtual FVector2f GetUV(int32 EID) const
 	{
-		return FVector2f(Collection->UV[EID]);
+		return FVector2f(Collection->UVs[EID][0]);
 	}
 
 	virtual void SetUV(int32 EID, FVector2f UVIn)
 	{
-		FVector2D& UV = Collection->UV[EID];
+		FVector2D& UV = Collection->UVs[EID][0];
 		UV.X = UVIn.X;
 		UV.Y = UVIn.Y;
 	}
@@ -299,7 +299,7 @@ struct FGeomFlatUVMesh
 
 	inline FVector3d GetVertex(int32 VID) const
 	{
-		const FVector2D& UV = Collection->UV[VID];
+		const FVector2D& UV = Collection->UVs[VID][0];
 		return FVector3d(UV.X, UV.Y, 0);
 	}
 
@@ -317,7 +317,7 @@ struct FGeomFlatUVMesh
 	}
 	inline int32 MaxVertexID() const
 	{
-		return Collection->UV.Num();
+		return Collection->UVs.Num();
 	}
 	inline int32 TriangleCount() const
 	{
@@ -325,7 +325,7 @@ struct FGeomFlatUVMesh
 	}
 	inline int32 VertexCount() const
 	{
-		return Collection->UV.Num();
+		return Collection->UVs.Num();
 	}
 	constexpr inline int32 GetShapeTimestamp() const
 	{
@@ -535,7 +535,7 @@ void TextureInternalSurfaces(
 			int32 TransformIdx = CollectionMeshes.Meshes[MeshIdx].TransformIndex;
 			FDynamicMesh3& Mesh = CollectionMeshes.Meshes[MeshIdx].AugMesh;
 			FMeshNormals::InitializeOverlayToPerVertexNormals(Mesh.Attributes()->PrimaryNormals(), true);
-			AugmentedDynamicMesh::InitializeOverlayToPerVertexUVs(Mesh);
+			AugmentedDynamicMesh::InitializeOverlayToPerVertexUVs(Mesh, Collection.NumUVLayers());
 			FDynamicMeshUVOverlay* UV = Mesh.Attributes()->PrimaryUV();
 			for (int TID : Mesh.TriangleIndicesItr())
 			{
diff --git a/Engine/Plugins/Experimental/PlanarCutPlugin/Source/PlanarCut/Private/GeometryMeshConversion.cpp b/Engine/Plugins/Experimental/PlanarCutPlugin/Source/PlanarCut/Private/GeometryMeshConversion.cpp
--- a/Engine/Plugins/Experimental/PlanarCutPlugin/Source/PlanarCut/Private/GeometryMeshConversion.cpp
+++ b/Engine/Plugins/Experimental/PlanarCutPlugin/Source/PlanarCut/Private/GeometryMeshConversion.cpp
@@ -61,33 +61,68 @@ namespace AugmentedDynamicMesh
 	FName TangentUAttribName = "TangentUAttrib";
 	FName TangentVAttribName = "TangentVAttrib";
 	FName VisibleAttribName = "VisibleAttrib";
+	enum
+	{
+		MAX_NUM_UV_CHANNELS = 8,
+	};
+	FName UVChannelNames[MAX_NUM_UV_CHANNELS] = {
+		"UVAttrib0",
+		"UVAttrib1",
+		"UVAttrib2",
+		"UVAttrib3",
+		"UVAttrib4",
+		"UVAttrib5",
+		"UVAttrib6",
+		"UVAttrib7"
+	};
 
-	void Augment(FDynamicMesh3& Mesh)
+	void EnableUVChannels(FDynamicMesh3& Mesh, int32 NumUVChannels, bool bResetExisting = false, bool bDisablePrevious = true)
+	{
+		Mesh.EnableAttributes();
+		if (!ensure(NumUVChannels <= MAX_NUM_UV_CHANNELS))
+		{
+			NumUVChannels = MAX_NUM_UV_CHANNELS;
+		}
+		for (int32 UVIdx = 0; UVIdx < NumUVChannels; UVIdx++)
+		{
+			if (!bResetExisting && Mesh.Attributes()->HasAttachedAttribute(UVChannelNames[UVIdx]))
+			{
+				continue;
+			}
+			Mesh.Attributes()->AttachAttribute(UVChannelNames[UVIdx], new TDynamicMeshVertexAttribute<float, 2>(&Mesh));
+		}
+		if (bDisablePrevious)
+		{
+			for (int32 UVIdx = NumUVChannels; UVIdx < MAX_NUM_UV_CHANNELS; UVIdx++)
+			{
+				Mesh.Attributes()->RemoveAttribute(UVChannelNames[UVIdx]);
+			}
+		}
+	}
+
+	void Augment(FDynamicMesh3& Mesh, int32 NumUVChannels)
 	{
 		Mesh.EnableVertexColors(FVector3f(1, 1, 1));
 		Mesh.EnableVertexNormals(FVector3f::UnitZ());
-		Mesh.EnableVertexUVs(FVector2f(0, 0));
 		Mesh.EnableAttributes();
 		Mesh.Attributes()->EnableMaterialID();
-		ensure(Mesh.Attributes()->NumAttachedAttributes() == 0);
 		Mesh.Attributes()->AttachAttribute(TangentUAttribName, new TDynamicMeshVertexAttribute<float, 3>(&Mesh));
 		Mesh.Attributes()->AttachAttribute(TangentVAttribName, new TDynamicMeshVertexAttribute<float, 3>(&Mesh));
 		TDynamicMeshScalarTriangleAttribute<bool>* VisAttrib = new TDynamicMeshScalarTriangleAttribute<bool>(&Mesh);
 		VisAttrib->Initialize(true);
 		Mesh.Attributes()->AttachAttribute(VisibleAttribName, VisAttrib);
+		EnableUVChannels(Mesh, NumUVChannels);
 	}
 
 	bool IsAugmented(const FDynamicMesh3& Mesh)
 	{
 		return Mesh.HasAttributes()
-			&& Mesh.Attributes()->NumAttachedAttributes() == 3
 			&& Mesh.Attributes()->HasAttachedAttribute(TangentUAttribName)
 			&& Mesh.Attributes()->HasAttachedAttribute(TangentVAttribName)
 			&& Mesh.Attributes()->HasAttachedAttribute(VisibleAttribName)
 			&& Mesh.Attributes()->HasMaterialID()
 			&& Mesh.HasVertexColors()
-			&& Mesh.HasVertexNormals()
-			&& Mesh.HasVertexUVs();
+			&& Mesh.HasVertexNormals();
 	}
 
 	void SetDefaultAttributes(FDynamicMesh3& Mesh, bool bGlobalVisibility)
@@ -131,6 +166,52 @@ namespace AugmentedDynamicMesh
 		return Visible->GetValue(TID);
 	}
 
+	void SetUV(FDynamicMesh3& Mesh, int VID, FVector2f UV, int UVLayer)
+	{
+		if (!ensure(UVLayer < MAX_NUM_UV_CHANNELS))
+		{
+			return;
+		}
+		TDynamicMeshVertexAttribute<float, 2>* UVs =
+			static_cast<TDynamicMeshVertexAttribute<float, 2>*>(Mesh.Attributes()->GetAttachedAttribute(UVChannelNames[UVLayer]));
+		if (ensure(UVs))
+		{
+			UVs->SetValue(VID, UV);
+		}
+	}
+
+	void SetAllUV(FDynamicMesh3& Mesh, int VID, FVector2f UV, int NumUVLayers)
+	{
+		if (!ensure(NumUVLayers < MAX_NUM_UV_CHANNELS))
+		{
+			NumUVLayers = MAX_NUM_UV_CHANNELS;
+		}
+		for (int32 Layer = 0; Layer < NumUVLayers; Layer++)
+		{
+			TDynamicMeshVertexAttribute<float, 2>* UVs =
+				static_cast<TDynamicMeshVertexAttribute<float, 2>*>(Mesh.Attributes()->GetAttachedAttribute(UVChannelNames[Layer]));
+			if (ensure(UVs))
+			{
+				UVs->SetValue(VID, UV);
+			}
+		}
+	}
+
+	void GetUV(const FDynamicMesh3& Mesh, int VID, FVector2f& UV, int UVLayer)
+	{
+		if (!ensure(UVLayer < MAX_NUM_UV_CHANNELS))
+		{
+			UV = FVector2f::Zero();
+			return;
+		}
+		TDynamicMeshVertexAttribute<float, 2>* UVs =
+			static_cast<TDynamicMeshVertexAttribute<float, 2>*>(Mesh.Attributes()->GetAttachedAttribute(UVChannelNames[UVLayer]));
+		if (ensure(UVs))
+		{
+			UVs->GetValue(VID, UV);
+		}
+	}
+
 	void SetTangent(FDynamicMesh3& Mesh, int VID, FVector3f Normal, FVector3f TangentU, FVector3f TangentV)
 	{
 		checkSlow(IsAugmented(Mesh));
@@ -154,25 +235,30 @@ namespace AugmentedDynamicMesh
 		Vs->GetValue(VID, V);
 	}
 
-	void InitializeOverlayToPerVertexUVs(FDynamicMesh3& Mesh)
+	void InitializeOverlayToPerVertexUVs(FDynamicMesh3& Mesh, int32 NumUVLayers)
 	{
-		FDynamicMeshUVOverlay* UVs = Mesh.Attributes()->PrimaryUV();
-		UVs->ClearElements();
-		TArray<int> VertToUVMap;
-		VertToUVMap.SetNumUninitialized(Mesh.MaxVertexID());
-		for (int VID : Mesh.VertexIndicesItr())
+		Mesh.Attributes()->SetNumUVLayers(NumUVLayers);
+		for (int UVLayer = 0; UVLayer < NumUVLayers; ++UVLayer)
 		{
-			FVector2f UV = Mesh.GetVertexUV(VID);
-			int UVID = UVs->AppendElement(UV);
-			VertToUVMap[VID] = UVID;
-		}
-		for (int TID : Mesh.TriangleIndicesItr())
-		{
-			FIndex3i Tri = Mesh.GetTriangle(TID);
-			Tri.A = VertToUVMap[Tri.A];
-			Tri.B = VertToUVMap[Tri.B];
-			Tri.C = VertToUVMap[Tri.C];
-			UVs->SetTriangle(TID, Tri);
+			FDynamicMeshUVOverlay* UVs = Mesh.Attributes()->GetUVLayer(UVLayer);
+			UVs->ClearElements();
+			TArray<int> VertToUVMap;
+			VertToUVMap.SetNumUninitialized(Mesh.MaxVertexID());
+			for (int VID : Mesh.VertexIndicesItr())
+			{
+				FVector2f UV;
+				GetUV(Mesh, VID, UV, UVLayer);
+				int UVID = UVs->AppendElement(UV);
+				VertToUVMap[VID] = UVID;
+			}
+			for (int TID : Mesh.TriangleIndicesItr())
+			{
+				FIndex3i Tri = Mesh.GetTriangle(TID);
+				Tri.A = VertToUVMap[Tri.A];
+				Tri.B = VertToUVMap[Tri.B];
+				Tri.C = VertToUVMap[Tri.C];
+				UVs->SetTriangle(TID, Tri);
+			}
 		}
 	}
 
@@ -215,7 +301,7 @@ namespace AugmentedDynamicMesh
 
 		// Copy per-vertex UVs to a UV overlay, because that's what the tangents code uses
 		// (TODO: consider making a tangent computation path that uses vertex normals / UVs)
-		InitializeOverlayToPerVertexUVs(Mesh);
+		InitializeOverlayToPerVertexUVs(Mesh, 1); // tangents only use the first layer
 		FDynamicMeshUVOverlay* UVs = Mesh.Attributes()->PrimaryUV();
 
 		FComputeTangentsOptions Options;
@@ -367,20 +453,21 @@ namespace AugmentedDynamicMesh
 }
 
 
-void SetGeometryCollectionAttributes(FDynamicMesh3& Mesh)
+void SetGeometryCollectionAttributes(FDynamicMesh3& Mesh, int32 NumUVLayers)
 {
-	AugmentedDynamicMesh::Augment(Mesh);
+	AugmentedDynamicMesh::Augment(Mesh, NumUVLayers);
 }
 
 
-FCellMeshes::FCellMeshes(const FPlanarCells& Cells, FAxisAlignedBox3d DomainBounds, double Grout, double ExtendDomain, bool bIncludeOutsideCell)
+FCellMeshes::FCellMeshes(int32 NumUVLayersIn, const FPlanarCells& Cells, FAxisAlignedBox3d DomainBounds, double Grout, double ExtendDomain, bool bIncludeOutsideCell)
 {
-	Init(Cells, DomainBounds, Grout, ExtendDomain, bIncludeOutsideCell);
+	Init(NumUVLayersIn, Cells, DomainBounds, Grout, ExtendDomain, bIncludeOutsideCell);
 }
 
 
-FCellMeshes::FCellMeshes(FDynamicMesh3& SingleCutter, const FInternalSurfaceMaterials& Materials, TOptional<FTransform> Transform)
+FCellMeshes::FCellMeshes(int32 NumUVLayersIn, FDynamicMesh3& SingleCutter, const FInternalSurfaceMaterials& Materials, TOptional<FTransform> Transform)
 {
+	NumUVLayers = NumUVLayersIn;
 	SetNumCells(2);
 
 	if (Transform.IsSet())
@@ -391,7 +478,7 @@ FCellMeshes::FCellMeshes(FDynamicMesh3& SingleCutter, const FInternalSurfaceMate
 	// Mesh should already be augmented
 	if (!ensure(AugmentedDynamicMesh::IsAugmented(SingleCutter)))
 	{
-		AugmentedDynamicMesh::Augment(SingleCutter);
+		AugmentedDynamicMesh::Augment(SingleCutter, NumUVLayers);
 	}
 
 	CellMeshes[0].AugMesh = SingleCutter;
@@ -527,8 +614,9 @@ void FCellMeshes::ApplyNoise(FDynamicMesh3& Mesh, FVector3d Normal, const FNoise
 }
 
 
-void FCellMeshes::Init(const FPlanarCells& Cells, FAxisAlignedBox3d DomainBounds, double Grout, double ExtendDomain, bool bIncludeOutsideCell)
+void FCellMeshes::Init(int32 NumUVLayersIn, const FPlanarCells& Cells, FAxisAlignedBox3d DomainBounds, double Grout, double ExtendDomain, bool bIncludeOutsideCell)
 {
+	NumUVLayers = NumUVLayersIn;
 	InitEmpty();
 
 	float GlobalUVScale = Cells.InternalSurfaceMaterials.GlobalUVScale;
@@ -621,7 +709,7 @@ void FCellMeshes::ApplyGeneralGrout(double Grout)
 		{
 			// if the grout scale factor would be ~zero or negative, just clear the mesh instead
 			Mesh.Clear();
-			AugmentedDynamicMesh::Augment(Mesh);
+			AugmentedDynamicMesh::Augment(Mesh, NumUVLayers);
 		}
 		else
 		{
@@ -634,7 +722,7 @@ void FCellMeshes::ApplyGeneralGrout(double Grout)
 	{
 		FDynamicMesh3& OutsideMesh = CellMeshes[OutsideCellIndex].AugMesh;
 		OutsideMesh.Clear();
-		AugmentedDynamicMesh::Augment(OutsideMesh);
+		AugmentedDynamicMesh::Augment(OutsideMesh, NumUVLayers);
 		FDynamicMeshEditor OutsideMeshEditor(&OutsideMesh);
 		for (int MeshIdx = 0; MeshIdx < CellMeshes.Num(); MeshIdx++)
 		{
@@ -689,7 +777,7 @@ void FCellMeshes::CreateMeshesForBoundedPlanesWithoutNoise(int NumCells, const F
 		FFrame3d PlaneFrame(Cells.Planes[PlaneIdx]);
 		FVertexInfo PlaneVertInfo;
 		PlaneVertInfo.bHaveC = true;
-		PlaneVertInfo.bHaveUV = true;
+		PlaneVertInfo.bHaveUV = false;
 		PlaneVertInfo.bHaveN = true;
 		PlaneVertInfo.Color = FVector3f(1, 1, 1);
 		int VertStart[2]{ -1, -1 };
@@ -712,8 +800,9 @@ void FCellMeshes::CreateMeshesForBoundedPlanesWithoutNoise(int NumCells, const F
 			for (int BoundaryVertex : PlaneBoundary)
 			{
 				PlaneVertInfo.Position = FVector3d(Cells.PlaneBoundaryVertices[BoundaryVertex]);
-				PlaneVertInfo.UV = (FVector2f(PlaneFrame.ToPlaneUV(PlaneVertInfo.Position)) - MinUV) * GlobalUVScale;
-				Meshes[MeshIdx]->AppendVertex(PlaneVertInfo);
+				FVector2f UV = (FVector2f(PlaneFrame.ToPlaneUV(PlaneVertInfo.Position)) - MinUV) * GlobalUVScale;
+				int VID = Meshes[MeshIdx]->AppendVertex(PlaneVertInfo);
+				AugmentedDynamicMesh::SetAllUV(*Meshes[MeshIdx], VID, UV, NumUVLayers);
 			}
 		}
 
@@ -745,7 +834,9 @@ void FCellMeshes::CreateMeshesForBoundedPlanesWithoutNoise(int NumCells, const F
 			FPolygon2f Polygon;
 			for (int V = 0; V < PlaneBoundary.Num(); V++)
 			{
-				Polygon.AppendVertex(Meshes[0]->GetVertexUV(VertStart[0] + V));
+				FVector2f UV;
+				AugmentedDynamicMesh::GetUV(*Meshes[0], VertStart[0] + V, UV, 0);
+				Polygon.AppendVertex(UV);
 			}
 
 			FGeneralPolygon2f GeneralPolygon(Polygon);
@@ -799,7 +890,7 @@ void FCellMeshes::CreateMeshesForBoundedPlanesWithNoise(int NumCells, const FPla
 	FName OriginalPositionAttribute = "OriginalPosition";
 	for (FDynamicMesh3& PlaneMesh : PlaneMeshes)
 	{
-		PlaneMesh.EnableVertexUVs(FVector2f(0, 0));
+		AugmentedDynamicMesh::EnableUVChannels(PlaneMesh, NumUVLayers);
 		PlaneMesh.EnableVertexNormals(FVector3f::UnitZ());
 		PlaneMesh.EnableVertexColors(FVector3f(1, 1, 1));
 		PlaneMesh.EnableAttributes();
@@ -854,10 +945,10 @@ void FCellMeshes::CreateMeshesForBoundedPlanesWithNoise(int NumCells, const FPla
 			FFrame3d PlaneFrame(Cells.Planes[PlaneIdx]);
 			FVertexInfo PlaneVertInfo;
 			PlaneVertInfo.bHaveC = true;
-			PlaneVertInfo.bHaveUV = true;
+			PlaneVertInfo.bHaveUV = false;
 			PlaneVertInfo.bHaveN = true;
 			PlaneVertInfo.Normal = Normal;
-			PlaneVertInfo.UV = FVector2f(0, 0); // UVs will be set below, after noise is added
+			// UVs will be set below, after noise is added
 			PlaneVertInfo.Color = FVector3f(1, 1, 1);
 
 			FPolygon2f Polygon;
@@ -1133,7 +1224,10 @@ void FCellMeshes::CreateMeshesForBoundedPlanesWithNoise(int NumCells, const FPla
 				for (int Idx = 0; Idx < 3; Idx++)
 				{
 					FVector2f UV = ((FVector2f)PlaneFrames[PlaneIdx].ToPlaneUV(Mesh.GetVertex(Tri[Idx])) - PlaneMinUVs[PlaneIdx]) * GlobalUVScale;
-					Mesh.SetVertexUV(Tri[Idx], UV);
+					for (int UVLayerIdx = 0; UVLayerIdx < NumUVLayers; UVLayerIdx++)
+					{
+						AugmentedDynamicMesh::SetUV(Mesh, Tri[Idx], UV, UVLayerIdx);
+					}
 				}
 			}
 		}
@@ -1162,10 +1256,10 @@ void FCellMeshes::CreateMeshesForSinglePlane(const FPlanarCells& Cells, const FA
 	//	// TODO: early out for plane that doesn't even intersect the domain bounding box?
 	//}
 
-	FDynamicMesh3 PlaneMesh(true, true, true, false);
+	FDynamicMesh3 PlaneMesh(true, true, false, false);
+	AugmentedDynamicMesh::EnableUVChannels(PlaneMesh, NumUVLayers);
 	FVertexInfo PlaneVertInfo;
 	PlaneVertInfo.bHaveC = true;
-	PlaneVertInfo.bHaveUV = true;
 	PlaneVertInfo.bHaveN = true;
 	PlaneVertInfo.Color = FVector3f(1, 1, 1);
 	PlaneVertInfo.Normal = -FVector3f(Plane.GetNormal());
@@ -1173,8 +1267,9 @@ void FCellMeshes::CreateMeshesForSinglePlane(const FPlanarCells& Cells, const FA
 	for (int CornerIdx = 0; CornerIdx < 4; CornerIdx++)
 	{
 		PlaneVertInfo.Position = PlaneFrame.FromPlaneUV(XYRange.GetCorner(CornerIdx));
-		PlaneVertInfo.UV = FVector2f(XYRange.GetCorner(CornerIdx) - XYRange.Min) * GlobalUVScale;
-		PlaneMesh.AppendVertex(PlaneVertInfo);
+		FVector2f UV = FVector2f(XYRange.GetCorner(CornerIdx) - XYRange.Min) * GlobalUVScale;
+		int VID = PlaneMesh.AppendVertex(PlaneVertInfo);
+		AugmentedDynamicMesh::SetAllUV(PlaneMesh, VID, UV, NumUVLayers);
 	}
 	PlaneMesh.AppendTriangle(0, 1, 2);
 	PlaneMesh.AppendTriangle(0, 2, 3);
@@ -1221,6 +1316,7 @@ void FCellMeshes::CreateMeshesForSinglePlane(const FPlanarCells& Cells, const FA
 		{
 			Meshes[Side] = &CellMeshes[Side].AugMesh;
 			*Meshes[Side] = PlaneMesh;
+
 			double Offset = ZRange.Max;
 			TArray<int> CapBoundary, CapBoundaryCornerIndices;
 
@@ -1289,8 +1385,9 @@ void FCellMeshes::CreateMeshesForSinglePlane(const FPlanarCells& Cells, const FA
 	for (int CellIdx = 0; CellIdx < CellMeshes.Num(); CellIdx++)
 	{
 		FDynamicMesh3& Mesh = CellMeshes[CellIdx].AugMesh;
+
 		// re-enable tangents and visibility attributes, since these are lost when we set the mesh to a copy of the plane mesh
-		AugmentedDynamicMesh::Augment(Mesh);
+		AugmentedDynamicMesh::Augment(Mesh, NumUVLayers);
 
 		// Set all material IDs to the one plane's corresponding material ID
 		for (int TID : Mesh.TriangleIndicesItr())
@@ -1303,6 +1400,7 @@ void FCellMeshes::CreateMeshesForSinglePlane(const FPlanarCells& Cells, const FA
 
 void FDynamicMeshCollection::Init(const FGeometryCollection* Collection, const TArrayView<const int32>& TransformIndices, FTransform TransformCollection, bool bSaveIsolatedVertices)
 {
+	int32 NumUVLayers = Collection->NumUVLayers();
 	Meshes.Reset();
 	Bounds = FAxisAlignedBox3d::Empty();
 
@@ -1316,7 +1414,7 @@ void FDynamicMeshCollection::Init(const FGeometryCollection* Collection, const T
 
 		FTransform3d CollectionToLocal = FTransform3d(GeometryCollectionAlgo::GlobalMatrix(Collection->Transform, Collection->Parent, TransformIdx) * TransformCollection);
 
-		int32 AddedMeshIdx = Meshes.Add(new FMeshData);
+		int32 AddedMeshIdx = Meshes.Add(new FMeshData(NumUVLayers));
 		FMeshData& MeshData = Meshes[AddedMeshIdx];
 		MeshData.TransformIndex = TransformIdx;
 		MeshData.ToCollection = FTransform(CollectionToLocal.Inverse());
@@ -1326,6 +1424,9 @@ void FDynamicMeshCollection::Init(const FGeometryCollection* Collection, const T
 		Mesh.EnableAttributes();
 		Mesh.Attributes()->EnableMaterialID();
 
+		//int32 NumUVLayers = Collection->NumUVLayers();
+		//AugmentedDynamicMesh::InitializeOverlayToPerVertexUVs(Mesh, NumUVLayers);
+
 		int32 VertexStart = Collection->VertexStart[GeometryIdx];
 		int32 VertexCount = Collection->VertexCount[GeometryIdx];
 		int32 FaceCount = Collection->FaceCount[GeometryIdx];
@@ -1333,17 +1434,21 @@ void FDynamicMeshCollection::Init(const FGeometryCollection* Collection, const T
 		FVertexInfo VertexInfo;
 		VertexInfo.bHaveC = true;
 		VertexInfo.bHaveN = true;
-		VertexInfo.bHaveUV = true;
+		VertexInfo.bHaveUV = false;
 		for (int32 Idx = VertexStart, N = VertexStart + VertexCount; Idx < N; Idx++)
 		{
 			VertexInfo.Position = CollectionToLocal.TransformPosition(FVector3d(Collection->Vertex[Idx]));
-			VertexInfo.UV = FVector2f(Collection->UV[Idx]);
 			VertexInfo.Color = FVector3f(Collection->Color[Idx]);
 			VertexInfo.Normal = (FVector3f)CollectionToLocal.TransformVectorNoScale(FVector3d(Collection->Normal[Idx]));
 			int VID = Mesh.AppendVertex(VertexInfo);
 			AugmentedDynamicMesh::SetTangent(Mesh, VID, VertexInfo.Normal,
 				(FVector3f)CollectionToLocal.TransformVectorNoScale(FVector3d(Collection->TangentU[Idx])),
 				(FVector3f)CollectionToLocal.TransformVectorNoScale(FVector3d(Collection->TangentV[Idx])));
+			
+			for (int32 UVLayer = 0; UVLayer < NumUVLayers; ++UVLayer)
+			{
+				AugmentedDynamicMesh::SetUV(Mesh, VID, Collection->UVs[Idx][UVLayer], UVLayer);
+			}
 		}
 		FIntVector VertexOffset(VertexStart, VertexStart, VertexStart);
 		for (int32 Idx = Collection->FaceStart[GeometryIdx], N = Collection->FaceStart[GeometryIdx] + FaceCount; Idx < N; Idx++)
@@ -1362,6 +1467,12 @@ void FDynamicMeshCollection::Init(const FGeometryCollection* Collection, const T
 						Mesh.GetVertexNormal(NewVID), // TODO: we don't actually use the vertex normal; consider removing this arg from the function entirely
 						(FVector3f)CollectionToLocal.TransformVectorNoScale(FVector3d(Collection->TangentU[SrcIdx])),
 						(FVector3f)CollectionToLocal.TransformVectorNoScale(FVector3d(Collection->TangentV[SrcIdx])));
+
+					for (int32 UVLayer = 0; UVLayer < NumUVLayers; ++UVLayer)
+					{
+						AugmentedDynamicMesh::SetUV(Mesh, NewVID, Collection->UVs[SrcIdx][UVLayer], UVLayer);
+					}
+
 					NewTri[SubIdx] = NewVID;
 				}
 				TID = Mesh.AppendTriangle(NewTri, 0);
@@ -1416,6 +1527,8 @@ int32 FDynamicMeshCollection::CutWithMultiplePlanes(
 
 	bool bHasGrout = Grout > 0;
 
+	int32 NumUVLayers = Collection->NumUVLayers();
+
 	if (bHasGrout)
 	{
 		// For multi-plane cuts with grout specifically, the easiest path seems to be:
@@ -1424,7 +1537,7 @@ int32 FDynamicMeshCollection::CutWithMultiplePlanes(
 		// 3. Use the generic CutWithCellMeshes path, where that grout mesh is both the inner and outside cell mesh
 		//    (Note the outside cell mesh is subtracted, not intersected)
 		//    (Note this relies on island splitting to separate all the pieces afterwards.)
-		FCellMeshes GroutCells;
+		FCellMeshes GroutCells(NumUVLayers);
 		GroutCells.SetNumCells(2);
 		FDynamicMesh3& GroutMesh = GroutCells.CellMeshes[0].AugMesh;
 		FDynamicMeshEditor GroutAppender(&GroutMesh);
@@ -1434,7 +1547,7 @@ int32 FDynamicMeshCollection::CutWithMultiplePlanes(
 			EnterProgressFrame(.5);
 			FPlanarCells PlaneCells(Planes[PlaneIdx]);
 			PlaneCells.InternalSurfaceMaterials = InternalSurfaceMaterials;
-			FCellMeshes PlaneGroutMesh;
+			FCellMeshes PlaneGroutMesh(NumUVLayers);
 			PlaneGroutMesh.MakeOnlyPlanarGroutCell(PlaneCells, Bounds, Grout);
 			GroutAppender.AppendMesh(&PlaneGroutMesh.CellMeshes[0].AugMesh, IndexMaps);
 		}
@@ -1540,7 +1653,7 @@ int32 FDynamicMeshCollection::CutWithMultiplePlanes(
 		FPlanarCells PlaneCells(Planes[PlaneIdx]);
 		PlaneCells.InternalSurfaceMaterials = InternalSurfaceMaterials;
 		double OnePercentExtend = Bounds.MaxDim() * .01;
-		FCellMeshes CellMeshes(PlaneCells, Bounds, 0, OnePercentExtend, false);
+		FCellMeshes CellMeshes(NumUVLayers, PlaneCells, Bounds, 0, OnePercentExtend, false);
 
 		// TODO: we could do these cuts in parallel (will takes some rework of the proximity and how results are added to the ToCut array)
 		for (int32 ToCutIdx = 0, ToCutNum = ToCut.Num(); ToCutIdx < ToCutNum; ToCutIdx++)
@@ -1562,7 +1675,7 @@ int32 FDynamicMeshCollection::CutWithMultiplePlanes(
 			TArray<TUniquePtr<FMeshData>> BoolResults;
 			for (int ResultIdx = 0; ResultIdx < 2; ResultIdx++)
 			{
-				BoolResults.Add(MakeUnique<FMeshData>());
+				BoolResults.Add(MakeUnique<FMeshData>(NumUVLayers));
 				BoolResults[ResultIdx]->TransformIndex = TransformIndex;
 				BoolResults[ResultIdx]->ToCollection = ToCollection;
 			}
@@ -1764,7 +1877,6 @@ int32 FDynamicMeshCollection::CutWithCellMeshes(const FInternalSurfaceMaterials&
 				{
 					BooleanResults[CellIdx] = MakeUnique<FDynamicMesh3>();
 					FDynamicMesh3& AugBoolResult = *BooleanResults[CellIdx];
-
 					FMeshBoolean::EBooleanOp Op = FMeshBoolean::EBooleanOp::Intersect;
 					if (CellIdx == CellMeshes.OutsideCellIndex)
 					{
@@ -2078,6 +2190,8 @@ bool FDynamicMeshCollection::UpdateCollection(const FTransform& ToCollection, FD
 	int32 OldVertexCount = Output.VertexCount[GeometryIdx];
 	int32 OldTriangleCount = Output.FaceCount[GeometryIdx];
 
+	int32 UVLayerCount = Output.NumUVLayers();
+
 	int32 NewVertexCount = Mesh.VertexCount();
 	int32 NewTriangleCount = Mesh.TriangleCount();
 
@@ -2096,7 +2210,14 @@ bool FDynamicMeshCollection::UpdateCollection(const FTransform& ToCollection, FD
 		int32 CopyToIdx = VerticesStart + VID;
 		Output.Vertex[CopyToIdx] = ToCollection.TransformPosition(FVector(Mesh.GetVertex(VID)));
 		Output.Normal[CopyToIdx] = ToCollection.TransformVectorNoScale(FVector(Mesh.GetVertexNormal(VID)));
-		Output.UV[CopyToIdx] = FVector2D(Mesh.GetVertexUV(VID));
+		
+		for (int32 UVLayer = 0; UVLayer < UVLayerCount; ++UVLayer)
+		{
+			FVector2f UV;
+			AugmentedDynamicMesh::GetUV(Mesh, VID, UV, UVLayer);
+			Output.UVs[CopyToIdx][UVLayer] = FVector2D(UV);
+		}
+		
 		FVector3f TangentU, TangentV;
 		AugmentedDynamicMesh::GetTangent(Mesh, VID, TangentU, TangentV);
 		Output.TangentU[CopyToIdx] = ToCollection.TransformVectorNoScale(FVector(TangentU));
@@ -2151,6 +2272,7 @@ int32 FDynamicMeshCollection::AppendToCollection(const FTransform& ToCollection,
 	int32 NewGeometryStartIdx = Output.FaceStart.Num();
 	int32 OriginalVertexNum = Output.Vertex.Num();
 	int32 OriginalFaceNum = Output.Indices.Num();
+	int32 UVLayerCount = Output.NumUVLayers();
 
 	int32 GeometryIdx = Output.AddElements(1, FGeometryCollection::GeometryGroup);
 	int32 TransformIdx = Output.AddElements(1, FGeometryCollection::TransformGroup);
@@ -2185,7 +2307,15 @@ int32 FDynamicMeshCollection::AppendToCollection(const FTransform& ToCollection,
 		int32 CopyToIdx = VerticesStart + VID;
 		Output.Vertex[CopyToIdx] = ToCollection.TransformPosition(FVector(Mesh.GetVertex(VID)));
 		Output.Normal[CopyToIdx] = ToCollection.TransformVectorNoScale(FVector(Mesh.GetVertexNormal(VID)));
-		Output.UV[CopyToIdx] = FVector2D(Mesh.GetVertexUV(VID));
+		
+		Output.UVs[CopyToIdx].SetNum(UVLayerCount);
+		for (int32 UVLayer = 0; UVLayer < UVLayerCount; ++UVLayer)
+		{
+			FVector2f UV;
+			AugmentedDynamicMesh::GetUV(Mesh, VID, UV, UVLayer);
+			Output.UVs[CopyToIdx][UVLayer] = FVector2D(UV);
+		}
+
 		FVector3f TangentU, TangentV;
 		AugmentedDynamicMesh::GetTangent(Mesh, VID, TangentU, TangentV);
 		Output.TangentU[CopyToIdx] = ToCollection.TransformVectorNoScale(FVector(TangentU));
diff --git a/Engine/Plugins/Experimental/PlanarCutPlugin/Source/PlanarCut/Private/GeometryMeshConversion.h b/Engine/Plugins/Experimental/PlanarCutPlugin/Source/PlanarCut/Private/GeometryMeshConversion.h
--- a/Engine/Plugins/Experimental/PlanarCutPlugin/Source/PlanarCut/Private/GeometryMeshConversion.h
+++ b/Engine/Plugins/Experimental/PlanarCutPlugin/Source/PlanarCut/Private/GeometryMeshConversion.h
@@ -23,16 +23,18 @@ namespace PlanarCut
 /**
  * Add attributes necessary for a dynamic mesh to represent geometry from an FGeometryCollection
  */
-void PLANARCUT_API SetGeometryCollectionAttributes(UE::Geometry::FDynamicMesh3& Mesh);
+void PLANARCUT_API SetGeometryCollectionAttributes(UE::Geometry::FDynamicMesh3& Mesh, int32 NumUVLayers);
 
 // functions for DynamicMesh3 meshes that have FGeometryCollection attributes set
 namespace AugmentedDynamicMesh
 {
 	void PLANARCUT_API SetVisibility(UE::Geometry::FDynamicMesh3& Mesh, int TID, bool bIsVisible);
 	bool PLANARCUT_API GetVisibility(const UE::Geometry::FDynamicMesh3& Mesh, int TID);
+	void PLANARCUT_API SetUV(UE::Geometry::FDynamicMesh3& Mesh, int VID, FVector2f UV, int UVLayer);
+	void PLANARCUT_API GetUV(const UE::Geometry::FDynamicMesh3& Mesh, int VID, FVector2f& UV, int UVLayer);
 	void PLANARCUT_API SetTangent(UE::Geometry::FDynamicMesh3& Mesh, int VID, FVector3f Normal, FVector3f TangentU, FVector3f TangentV);
 	void PLANARCUT_API GetTangent(const UE::Geometry::FDynamicMesh3& Mesh, int VID, FVector3f& U, FVector3f& V);
-	void PLANARCUT_API InitializeOverlayToPerVertexUVs(UE::Geometry::FDynamicMesh3& Mesh);
+	void PLANARCUT_API InitializeOverlayToPerVertexUVs(UE::Geometry::FDynamicMesh3& Mesh, int32 NumUVLayers);
 	void PLANARCUT_API InitializeOverlayToPerVertexTangents(UE::Geometry::FDynamicMesh3& Mesh);
 	void PLANARCUT_API ComputeTangents(UE::Geometry::FDynamicMesh3& Mesh, bool bOnlyOddMaterials, const TArrayView<const int32>& WhichMaterials, bool bRecomputeNormals = true);
 	void PLANARCUT_API AddCollisionSamplesPerComponent(UE::Geometry::FDynamicMesh3& Mesh, double Spacing);
@@ -48,9 +50,9 @@ struct PLANARCUT_API FCellMeshes
 	{
 		UE::Geometry::FDynamicMesh3 AugMesh;
 
-		FCellInfo()
+		FCellInfo(int32 NumUVLayers = 1)
 		{
-			SetGeometryCollectionAttributes(AugMesh);
+			SetGeometryCollectionAttributes(AugMesh, NumUVLayers);
 		}
 
 		// TODO: compute spatial in advance?  (only useful if we rework mesh booleans to support it)
@@ -59,6 +61,7 @@ struct PLANARCUT_API FCellMeshes
 
 	TIndirectArray<FCellInfo> CellMeshes;
 	int32 OutsideCellIndex = -1;
+	int32 NumUVLayers = 1;
 
 	// Noise Offsets, to randomize where perlin noise is sampled
 	FVector NoiseOffsetX;
@@ -70,18 +73,18 @@ struct PLANARCUT_API FCellMeshes
 		CellMeshes.Reset();
 		for (int32 Idx = 0; Idx < NumMeshes; Idx++)
 		{
-			CellMeshes.Add(new FCellInfo);
+			CellMeshes.Add(new FCellInfo(NumUVLayers));
 		}
 	}
 
-	FCellMeshes()
+	FCellMeshes(int32 NumUVLayers)
 	{
 		InitEmpty();
 	}
 
-	FCellMeshes(const FPlanarCells& Cells, UE::Geometry::FAxisAlignedBox3d DomainBounds, double Grout, double ExtendDomain, bool bIncludeOutsideCell);
+	FCellMeshes(int32 NumUVLayers, const FPlanarCells& Cells, UE::Geometry::FAxisAlignedBox3d DomainBounds, double Grout, double ExtendDomain, bool bIncludeOutsideCell);
 
-	FCellMeshes(UE::Geometry::FDynamicMesh3& SingleCutter, const FInternalSurfaceMaterials& Materials, TOptional<FTransform> Transform);
+	FCellMeshes(int32 NumUVLayers, UE::Geometry::FDynamicMesh3& SingleCutter, const FInternalSurfaceMaterials& Materials, TOptional<FTransform> Transform);
 
 	// Special function to just make the "grout" part of the planar mesh cells
 	// Used to make the multi-plane cuts with grout easier to implement
@@ -123,7 +126,7 @@ struct PLANARCUT_API FCellMeshes
 		OutsideCellIndex = -1;
 	}
 
-	void Init(const FPlanarCells& Cells, UE::Geometry::FAxisAlignedBox3d DomainBounds, double Grout, double ExtendDomain, bool bIncludeOutsideCell);
+	void Init(int32 NumUVLayersIn, const FPlanarCells& Cells, UE::Geometry::FAxisAlignedBox3d DomainBounds, double Grout, double ExtendDomain, bool bIncludeOutsideCell);
 
 	void ApplyGeneralGrout(double Grout);
 
@@ -154,9 +157,9 @@ struct PLANARCUT_API FDynamicMeshCollection
 		int32 TransformIndex; // where the mesh was from in the geometry collection
 		FTransform ToCollection; // transform that need be applied to go back to the local space of the geometry collection
 
-		FMeshData()
+		FMeshData(int32 NumUVLayers)
 		{
-			SetGeometryCollectionAttributes(AugMesh);
+			SetGeometryCollectionAttributes(AugMesh, NumUVLayers);
 		}
 
 		FMeshData(const UE::Geometry::FDynamicMesh3& Mesh, int32 TransformIndex, FTransform ToCollection) : AugMesh(Mesh), TransformIndex(TransformIndex), ToCollection(ToCollection)
@@ -164,7 +167,7 @@ struct PLANARCUT_API FDynamicMeshCollection
 	};
 	TIndirectArray<FMeshData> Meshes;
 	UE::Geometry::FAxisAlignedBox3d Bounds;
-
+	
 	FDynamicMeshCollection(const FGeometryCollection* Collection, const TArrayView<const int32>& TransformIndices, FTransform TransformCollection, bool bSaveIsolatedVertices = false)
 	{
 		Init(Collection, TransformIndices, TransformCollection, bSaveIsolatedVertices);
diff --git a/Engine/Plugins/Experimental/PlanarCutPlugin/Source/PlanarCut/Private/PlanarCut.cpp b/Engine/Plugins/Experimental/PlanarCutPlugin/Source/PlanarCut/Private/PlanarCut.cpp
--- a/Engine/Plugins/Experimental/PlanarCutPlugin/Source/PlanarCut/Private/PlanarCut.cpp
+++ b/Engine/Plugins/Experimental/PlanarCutPlugin/Source/PlanarCut/Private/PlanarCut.cpp
@@ -111,11 +111,11 @@ void FInternalSurfaceMaterials::SetUVScaleFromCollection(const FGeometryCollecti
 	{
 		const FIntVector& Tri = Collection.Indices[FaceIdx];
 		WorldDistance += FVector::Distance(Collection.Vertex[Tri.X], Collection.Vertex[Tri.Y]);
-		UVDistance += FVector2D::Distance(Collection.UV[Tri.X], Collection.UV[Tri.Y]);
+		UVDistance += FVector2D::Distance(Collection.UVs[Tri.X][0], Collection.UVs[Tri.Y][0]);
 		WorldDistance += FVector::Distance(Collection.Vertex[Tri.Z], Collection.Vertex[Tri.Y]);
-		UVDistance += FVector2D::Distance(Collection.UV[Tri.Z], Collection.UV[Tri.Y]);
+		UVDistance += FVector2D::Distance(Collection.UVs[Tri.Z][0], Collection.UVs[Tri.Y][0]);
 		WorldDistance += FVector::Distance(Collection.Vertex[Tri.X], Collection.Vertex[Tri.Z]);
-		UVDistance += FVector2D::Distance(Collection.UV[Tri.X], Collection.UV[Tri.Z]);
+		UVDistance += FVector2D::Distance(Collection.UVs[Tri.X][0], Collection.UVs[Tri.Z][0]);
 	}
 
 	if (WorldDistance > 0)
@@ -589,7 +589,6 @@ int32 CutMultipleWithMultiplePlanes(
 	return NewGeomStartIdx;
 }
 
-
 // Cut multiple Geometry groups inside a GeometryCollection with PlanarCells, and add each cut cell back to the GeometryCollection as a new child of their source Geometry
 int32 CutMultipleWithPlanarCells(
 	FPlanarCells& Cells,
@@ -618,7 +617,7 @@ int32 CutMultipleWithPlanarCells(
 
 	FDynamicMeshCollection MeshCollection(&Source, TransformIndices, CollectionToWorld);
 	double OnePercentExtend = MeshCollection.Bounds.MaxDim() * .01;
-	FCellMeshes CellMeshes(Cells, MeshCollection.Bounds, Grout, OnePercentExtend, bIncludeOutsideCellInOutput);
+	FCellMeshes CellMeshes(Source.NumUVLayers(), Cells, MeshCollection.Bounds, Grout, OnePercentExtend, bIncludeOutsideCellInOutput);
 
 	int32 NewGeomStartIdx = -1;
 
@@ -628,7 +627,6 @@ int32 CutMultipleWithPlanarCells(
 	return NewGeomStartIdx;
 }
 
-
 int32 CutWithMesh(
 	FMeshDescription* CuttingMesh,
 	FTransform CuttingMeshTransform,
@@ -662,7 +660,9 @@ int32 CutWithMesh(
 	}
 
 	FDynamicMesh3 DynamicCuttingMesh; // version of mesh that is split apart at seams to be compatible w/ geometry collection, with corresponding attributes set
-	SetGeometryCollectionAttributes(DynamicCuttingMesh);
+	int32 NumUVLayers = Collection.NumUVLayers();
+	SetGeometryCollectionAttributes(DynamicCuttingMesh, NumUVLayers);
+
 	// Note: This conversion will likely go away, b/c I plan to switch over to doing the boolean operations on the fuller rep, but the code can be adapted
 	//		 to the dynamic mesh -> geometry collection conversion phase, as this same splitting will then need to happen there.
 	if (ensure(FullMesh.HasAttributes() && FullMesh.Attributes()->NumUVLayers() >= 1 && FullMesh.Attributes()->NumNormalLayers() == 3))
@@ -686,7 +686,7 @@ int32 CutWithMesh(
 			FVector3d Pos = FullMesh.GetVertex(VID);
 
 			ElIDsToVID.Reset();
-			FullMesh.EnumerateVertexTriangles(VID, [&FullMesh, &Triangles, &OutMesh, &NTB, &UV, &ElIDsToVID, Pos, VID](int32 TID)
+			FullMesh.EnumerateVertexTriangles(VID, [&FullMesh, &Triangles, &OutMesh, &NTB, &UV, &ElIDsToVID, Pos, VID, NumUVLayers](int32 TID)
 			{
 				FIndex3i InTri = FullMesh.GetTriangle(TID);
 				int VOnT = IndexUtil::FindTriIndex(VID, InTri);
@@ -705,11 +705,15 @@ int32 CutWithMesh(
 				else
 				{
 					FVector3f Normal = NTB[0]->GetElement(ElIDs.A);
-					FVertexInfo Info(Pos, Normal, FVector3f(1, 1, 1), UV->GetElement(ElIDs.D));
+					FVertexInfo Info(Pos, Normal, FVector3f(1, 1, 1));
 
 					int OutVID = OutMesh.AppendVertex(Info);
 					OutTri[VOnT] = OutVID;
 					AugmentedDynamicMesh::SetTangent(OutMesh, OutVID, Normal, NTB[1]->GetElement(ElIDs.B), NTB[2]->GetElement(ElIDs.C));
+					for (int32 UVLayerIdx = 0; UVLayerIdx < NumUVLayers; UVLayerIdx++)
+					{
+						AugmentedDynamicMesh::SetUV(OutMesh, OutVID, UV->GetElement(ElIDs.D), UVLayerIdx);
+					}
 					ElIDsToVID.Add(ElIDs, OutVID);
 				}
 			});
@@ -744,7 +748,7 @@ int32 CutWithMesh(
 	FTransform CollectionToWorld = TransformCollection.Get(FTransform::Identity);
 
 	FDynamicMeshCollection MeshCollection(&Collection, TransformIndices, CollectionToWorld);
-	FCellMeshes CellMeshes(DynamicCuttingMesh, InternalSurfaceMaterials, CuttingMeshTransform);
+	FCellMeshes CellMeshes(NumUVLayers, DynamicCuttingMesh, InternalSurfaceMaterials, CuttingMeshTransform);
 
 	TArray<TPair<int32, int32>> CellConnectivity;
 	CellConnectivity.Add(TPair<int32, int32>(0, -1)); // there's only one 'inside' cell (0), so all cut surfaces are connecting the 'inside' cell (0) to the 'outside' cell (-1)
@@ -793,7 +797,6 @@ int32 AddCollisionSampleVertices(double CollisionSampleSpacing, FGeometryCollect
 }
 
 
-
 void ConvertToMeshDescription(
 	FMeshDescription& MeshOut,
 	FTransform& TransformOut,
@@ -808,7 +811,7 @@ void ConvertToMeshDescription(
 	FDynamicMeshCollection MeshCollection(&Collection, TransformIndices, CellsToWorld);
 	
 	FDynamicMesh3 CombinedMesh;
-	SetGeometryCollectionAttributes(CombinedMesh);
+	SetGeometryCollectionAttributes(CombinedMesh, Collection.NumUVLayers());
 	CombinedMesh.Attributes()->EnableTangents();
 
 	int32 NumMeshes = MeshCollection.Meshes.Num();
@@ -825,7 +828,7 @@ void ConvertToMeshDescription(
 		}
 
 		FMeshNormals::InitializeOverlayToPerVertexNormals(Mesh.Attributes()->PrimaryNormals(), true);
-		AugmentedDynamicMesh::InitializeOverlayToPerVertexUVs(Mesh);
+		AugmentedDynamicMesh::InitializeOverlayToPerVertexUVs(Mesh, Collection.NumUVLayers());
 		AugmentedDynamicMesh::InitializeOverlayToPerVertexTangents(Mesh);
 
 		FMergeCoincidentMeshEdges EdgeMerge(&Mesh);
@@ -853,4 +856,4 @@ void ConvertToMeshDescription(
 
 	FDynamicMeshToMeshDescription Converter;
 	Converter.Convert(&CombinedMesh, MeshOut, true);
-}
+}
\ No newline at end of file
diff --git a/Engine/Source/Runtime/Experimental/Chaos/Private/GeometryCollection/GeometryCollection.cpp b/Engine/Source/Runtime/Experimental/Chaos/Private/GeometryCollection/GeometryCollection.cpp
--- a/Engine/Source/Runtime/Experimental/Chaos/Private/GeometryCollection/GeometryCollection.cpp
+++ b/Engine/Source/Runtime/Experimental/Chaos/Private/GeometryCollection/GeometryCollection.cpp
@@ -54,7 +54,7 @@ void FGeometryCollection::Construct()
 	// Vertices Group
 	AddExternalAttribute<FVector3f>("Vertex", FGeometryCollection::VerticesGroup, Vertex);
 	AddExternalAttribute<FVector3f>("Normal", FGeometryCollection::VerticesGroup, Normal);
-	AddExternalAttribute<FVector2D>("UV", FGeometryCollection::VerticesGroup, UV);
+	AddExternalAttribute<TArray<FVector2D>>("UVs", FGeometryCollection::VerticesGroup, UVs);
 	AddExternalAttribute<FLinearColor>("Color", FGeometryCollection::VerticesGroup, Color);
 	AddExternalAttribute<FVector3f>("TangentU", FGeometryCollection::VerticesGroup, TangentU);
 	AddExternalAttribute<FVector3f>("TangentV", FGeometryCollection::VerticesGroup, TangentV);
@@ -118,7 +118,7 @@ int32 FGeometryCollection::AppendGeometry(const FGeometryCollection & Element, i
 	int NumNewVertices = Element.NumElements(FGeometryCollection::VerticesGroup);
 	const TManagedArray<FVector3f>& ElementVertices = Element.Vertex;
 	const TManagedArray<FVector3f>& ElementNormals = Element.Normal;
-	const TManagedArray<FVector2D>& ElementUVs = Element.UV;
+	const TManagedArray<TArray<FVector2D>>& ElementUVs = Element.UVs;
 	const TManagedArray<FLinearColor>& ElementColors = Element.Color;
 	const TManagedArray<FVector3f>& ElementTangentUs = Element.TangentU;
 	const TManagedArray<FVector3f>& ElementTangentVs = Element.TangentV;
@@ -162,7 +162,7 @@ int32 FGeometryCollection::AppendGeometry(const FGeometryCollection & Element, i
 	int VerticesIndex = AddElements(NumNewVertices, FGeometryCollection::VerticesGroup);
 	TManagedArray<FVector3f>& Vertices = Vertex;
 	TManagedArray<FVector3f>& Normals = Normal;
-	TManagedArray<FVector2D>& UVs = UV;
+	TManagedArray<TArray<FVector2D>>& AllUVs = UVs;
 	TManagedArray<FLinearColor>& Colors = Color;
 	TManagedArray<FVector3f>& TangentUs = TangentU;
 	TManagedArray<FVector3f>& TangentVs = TangentV;
@@ -173,7 +173,10 @@ int32 FGeometryCollection::AppendGeometry(const FGeometryCollection & Element, i
 	{
 		Vertices[VerticesIndex + vdx] = ElementVertices[vdx];
 		Normals[VerticesIndex + vdx] = ElementNormals[vdx];
-		UVs[VerticesIndex + vdx] = ElementUVs[vdx];
+		for (int UVChannelIndex = 0; UVChannelIndex < Chaos::GeometryCollection::MAX_NUM_UV_CHANNELS; UVChannelIndex++)
+		{
+			AllUVs[VerticesIndex + vdx][UVChannelIndex] = ElementUVs[vdx][UVChannelIndex];
+		}
 		Colors[VerticesIndex + vdx] = ElementColors[vdx];
 		TangentUs[VerticesIndex + vdx] = ElementTangentUs[vdx];
 		TangentVs[VerticesIndex + vdx] = ElementTangentVs[vdx];
@@ -1031,9 +1034,42 @@ void FGeometryCollection::Serialize(Chaos::FChaosArchive& Ar)
 			Version = 7;
 		}
 
+		// Version 8 introduced multiple UVs.
+		if (Version < 8)
+		{
+			if (!HasAttribute("UVs", FGeometryCollection::VerticesGroup))
+			{
+				UE_LOG(FGeometryCollectionLogging, Warning, TEXT("GeometryCollection updated to multiple UV sets."));
+				AddAttribute<TArray<FVector2D>>("UVs", FGeometryCollection::VerticesGroup);				
+			}
+
+			TManagedArray<TArray<FVector2D>>* MultipleUVs = FindAttribute<TArray<FVector2D>>("UVs", FGeometryCollection::VerticesGroup);
+			if (NumUVLayers() < 1)
+			{
+				for (int32 VertIdx = 0; VertIdx < MultipleUVs->Num(); ++VertIdx)
+				{
+					(*MultipleUVs)[VertIdx].SetNum(1);
+				}
+			}
+
+			if (TManagedArray<FVector2D>* SingleUV = FindAttribute<FVector2D>("UV", FGeometryCollection::VerticesGroup))
+			{
+				for (int32 VertIdx = 0; VertIdx < MultipleUVs->Num(); ++VertIdx)
+				{
+					if (SingleUV)
+					{
+						(*MultipleUVs)[VertIdx][0] = (*SingleUV)[VertIdx];
+					}
+				}
+
+				RemoveAttribute("UV", FGeometryCollection::VerticesGroup);
+			}
+
+			// Structure is conditioned, now considered up to date.
+			Version = 8;
+		}
 	}
 }
-
 bool FGeometryCollection::HasContiguousVertices( ) const
 {
 	int32 NumTransforms = NumElements(FGeometryCollection::TransformGroup);
@@ -1169,6 +1205,26 @@ bool FGeometryCollection::HasContiguousRenderFaces() const
 
 	return true;
 }
+
+int32 FGeometryCollection::NumUVLayers() const
+{
+	if (NumElements(VerticesGroup) == 0)
+	{
+		return 0;
+	}
+
+	// A valid GeometryCollection has the same number of UVs on every vertex.
+	// For safety, we consider only the smallest number of layers.
+	int32 MinUVLayers = UVs[0].Num();
+	for (int32 VertIdx = 0; VertIdx < NumElements(FGeometryCollection::VerticesGroup); ++VertIdx)
+	{
+		MinUVLayers = FMath::Min(MinUVLayers, UVs[VertIdx].Num());
+	}
+
+	return MinUVLayers;
+}
+
+
 FGeometryCollection* FGeometryCollection::NewGeometryCollection(const TArray<float>& RawVertexArray, const TArray<int32>& RawIndicesArray, bool ReverseVertexOrder)
 {
 
@@ -1187,7 +1243,7 @@ FGeometryCollection* FGeometryCollection::NewGeometryCollection(const TArray<flo
 	TManagedArray<FVector3f>&  Normals = RestCollection->Normal;
 	TManagedArray<FVector3f>&  TangentU = RestCollection->TangentU;
 	TManagedArray<FVector3f>&  TangentV = RestCollection->TangentV;
-	TManagedArray<FVector2D>&  UVs = RestCollection->UV;
+	TManagedArray<TArray<FVector2D>>& UVs = RestCollection->UVs;
 	TManagedArray<FLinearColor>&  Colors = RestCollection->Color;
 	TManagedArray<FIntVector>&  Indices = RestCollection->Indices;
 	TManagedArray<bool>&  Visible = RestCollection->Visible;
@@ -1202,7 +1258,7 @@ FGeometryCollection* FGeometryCollection::NewGeometryCollection(const TArray<flo
 		Vertices[Idx] = FVector3f(RawVertexArray[3 * Idx], RawVertexArray[3 * Idx + 1], RawVertexArray[3 * Idx + 2]);
 		TempVertices += Vertices[Idx];
 
-		UVs[Idx] = FVector2D(0, 0);
+		UVs[Idx].SetNumZeroed(Chaos::GeometryCollection::MAX_NUM_UV_CHANNELS);
 		Colors[Idx] = FLinearColor::White;
 	}
 
@@ -1501,7 +1557,7 @@ FGeometryCollection* FGeometryCollection::NewGeometryCollection(const TArray<flo
 	TManagedArray<FVector3f>&  Normals = RestCollection->Normal;
 	TManagedArray<FVector3f>&  TangentU = RestCollection->TangentU;
 	TManagedArray<FVector3f>&  TangentV = RestCollection->TangentV;
-	TManagedArray<FVector2D>&  UVs = RestCollection->UV;
+	TManagedArray<TArray<FVector2D>>& UVs = RestCollection->UVs;
 	TManagedArray<FLinearColor>&  Colors = RestCollection->Color;
 	TManagedArray<int32>& BoneMap = RestCollection->BoneMap;
 	TManagedArray<FIntVector>&  Indices = RestCollection->Indices;
@@ -1521,7 +1577,7 @@ FGeometryCollection* FGeometryCollection::NewGeometryCollection(const TArray<flo
 		Vertices[Idx] = FVector3f(RawVertexArray[3 * Idx], RawVertexArray[3 * Idx + 1], RawVertexArray[3 * Idx + 2]);
 		BoneMap[Idx] = RawBoneMapArray[Idx];
 
-		UVs[Idx] = FVector2D(0, 0);
+		UVs[Idx].SetNumZeroed(Chaos::GeometryCollection::MAX_NUM_UV_CHANNELS);
 		Colors[Idx] = FLinearColor::White;
 	}
 
diff --git a/Engine/Source/Runtime/Experimental/Chaos/Private/GeometryCollection/GeometryCollectionUtility.cpp b/Engine/Source/Runtime/Experimental/Chaos/Private/GeometryCollection/GeometryCollectionUtility.cpp
--- a/Engine/Source/Runtime/Experimental/Chaos/Private/GeometryCollection/GeometryCollectionUtility.cpp
+++ b/Engine/Source/Runtime/Experimental/Chaos/Private/GeometryCollection/GeometryCollectionUtility.cpp
@@ -32,7 +32,7 @@ namespace GeometryCollection
 		TManagedArray<FVector3f>&  Normals = RestCollection->Normal;
 		TManagedArray<FVector3f>&  TangentU = RestCollection->TangentU;
 		TManagedArray<FVector3f>&  TangentV = RestCollection->TangentV;
-		TManagedArray<FVector2D>&  UVs = RestCollection->UV;
+		TManagedArray<TArray<FVector2D>>& UVs = RestCollection->UVs;
 		TManagedArray<FLinearColor>&  Colors = RestCollection->Color;
 		TManagedArray<FIntVector>&  Indices = RestCollection->Indices;
 		TManagedArray<bool>&  Visible = RestCollection->Visible;
@@ -66,14 +66,15 @@ namespace GeometryCollection
 		Normals[6] = FVector3f(-1.f, 1.f, 1.f).GetSafeNormal();
 		Normals[7] = FVector3f(1.f, 1.f, 1.f).GetSafeNormal();
 
-		UVs[0] = FVector2D(0, 0);
-		UVs[1] = FVector2D(1, 0);
-		UVs[2] = FVector2D(0, 1);
-		UVs[3] = FVector2D(1, 1);
-		UVs[4] = FVector2D(0, 0);
-		UVs[5] = FVector2D(1, 0);
-		UVs[6] = FVector2D(0, 1);
-		UVs[7] = FVector2D(1, 1);
+		const int32 UVChannel = 0;
+		UVs[0][UVChannel] = FVector2D(0, 0);
+		UVs[1][UVChannel] = FVector2D(1, 0);
+		UVs[2][UVChannel] = FVector2D(0, 1);
+		UVs[3][UVChannel] = FVector2D(1, 1);
+		UVs[4][UVChannel] = FVector2D(0, 0);
+		UVs[5][UVChannel] = FVector2D(1, 0);
+		UVs[6][UVChannel] = FVector2D(0, 1);
+		UVs[7][UVChannel] = FVector2D(1, 1);
 
 		Colors[0] = FLinearColor::White;
 		Colors[1] = FLinearColor::White;
diff --git a/Engine/Source/Runtime/Experimental/Chaos/Private/GeometryCollection/ManagedArrayCollection.cpp b/Engine/Source/Runtime/Experimental/Chaos/Private/GeometryCollection/ManagedArrayCollection.cpp
--- a/Engine/Source/Runtime/Experimental/Chaos/Private/GeometryCollection/ManagedArrayCollection.cpp
+++ b/Engine/Source/Runtime/Experimental/Chaos/Private/GeometryCollection/ManagedArrayCollection.cpp
@@ -12,7 +12,7 @@ int8 FManagedArrayCollection::Invalid = INDEX_NONE;
 
 FManagedArrayCollection::FManagedArrayCollection()
 {
-	Version = 5;
+	Version = 8;
 }
 
 void FManagedArrayCollection::AddGroup(FName Group)
diff --git a/Engine/Source/Runtime/Experimental/Chaos/Public/GeometryCollection/GeometryCollection.h b/Engine/Source/Runtime/Experimental/Chaos/Public/GeometryCollection/GeometryCollection.h
--- a/Engine/Source/Runtime/Experimental/Chaos/Public/GeometryCollection/GeometryCollection.h
+++ b/Engine/Source/Runtime/Experimental/Chaos/Public/GeometryCollection/GeometryCollection.h
@@ -9,6 +9,14 @@
 namespace Chaos
 {
 	class FChaosArchive;
+
+	namespace GeometryCollection
+	{
+		enum 
+		{
+			MAX_NUM_UV_CHANNELS = 8,
+		};
+	}
 }
 
 /**
@@ -38,7 +46,7 @@ public:
 		*			FVectorArray      Vertex         = GetAttribute<FVector3f>("Vertex", VerticesGroup)
 		*			FInt32Array       BoneMap        = GetAttribute<Int32>("BoneMap", VerticesGroup, {"Transform"})
 		*           FVectorArray      Normal         = GetAttribute<FVector3f>("Normal", MaterialGroup)
-		*           FVector2DArray    UV             = GetAttribute<FVector2D>("UV", MaterialGroup)
+		*			FVector2DArray    UVs            = GetAttribute<TArray<FVector2D>>("UVs", MaterialGroup)
 		*           FVectorArray      TangentU       = GetAttribute<FVector3f>("TangentU", MaterialGroup)
 		*           FVectorArray      TangentV       = GetAttribute<FVector3f>("TangentV", MaterialGroup)
 		*           FLinearColorArray Color          = GetAttribute<FLinearColor>("Color", MaterialGroup)
@@ -208,6 +216,9 @@ public:
 	/** Returns true if the render faces are contiguous*/
 	bool HasContiguousRenderFaces() const;
 
+	/** Returns number of UV layers represented by UV array. A Valid Geometry Collection has the same count for every vertex */
+	int32 NumUVLayers() const;
+
 	FORCEINLINE bool IsGeometry(int32 Element) const { return TransformToGeometryIndex[Element] != INDEX_NONE; }
 	FORCEINLINE bool IsClustered(int32 Element) const { const TManagedArray<int32>& SimType = SimulationType;  return !!(SimType[Element] == ESimulationTypes::FST_Clustered); }
 	FORCEINLINE bool IsRigid(int32 Element) const { const TManagedArray<int32>& SimType = SimulationType;  return !!(SimType[Element] == ESimulationTypes::FST_Rigid); }
@@ -247,13 +258,14 @@ public:
 	TManagedArray<int32>		ExemplarIndex;
 
 	// Vertices Group
-	TManagedArray<FVector3f>		Vertex;
-	TManagedArray<FVector2D>    UV;
-	TManagedArray<FLinearColor> Color;
-	TManagedArray<FVector3f>      TangentU;
-	TManagedArray<FVector3f>      TangentV;
-	TManagedArray<FVector3f>      Normal;
-	TManagedArray<int32>        BoneMap;
+	TManagedArray<FVector3f>		 Vertex;
+	// Outer array is the array of vertices, inner array is the uv channels
+	TManagedArray<TArray<FVector2D>> UVs;
+	TManagedArray<FLinearColor>      Color;
+	TManagedArray<FVector3f>         TangentU;
+	TManagedArray<FVector3f>         TangentV;
+	TManagedArray<FVector3f>         Normal;
+	TManagedArray<int32>             BoneMap;
 
 	// Faces Group
 	TManagedArray<FIntVector>   Indices;
diff --git a/Engine/Source/Runtime/Experimental/Chaos/Public/GeometryCollection/GeometryCollectionUtility.h b/Engine/Source/Runtime/Experimental/Chaos/Public/GeometryCollection/GeometryCollectionUtility.h
--- a/Engine/Source/Runtime/Experimental/Chaos/Public/GeometryCollection/GeometryCollectionUtility.h
+++ b/Engine/Source/Runtime/Experimental/Chaos/Public/GeometryCollection/GeometryCollectionUtility.h
@@ -39,7 +39,7 @@ namespace GeometryCollection
 		TManagedArray<FVector3f>& Normals = RestCollection->Normal;
 		TManagedArray<FVector3f>& TangentU = RestCollection->TangentU;
 		TManagedArray<FVector3f>& TangentV = RestCollection->TangentV;
-		TManagedArray<FVector2D>& UVs = RestCollection->UV;
+		TManagedArray<TArray<FVector2D>>& UVs = RestCollection->UVs;
 		TManagedArray<FLinearColor>& Colors = RestCollection->Color;
 		TManagedArray<FIntVector>& Indices = RestCollection->Indices;
 		TManagedArray<bool>& Visible = RestCollection->Visible;
diff --git a/Engine/Source/Runtime/Experimental/Chaos/Public/GeometryCollection/ManagedArrayTypeValues.inl b/Engine/Source/Runtime/Experimental/Chaos/Public/GeometryCollection/ManagedArrayTypeValues.inl
--- a/Engine/Source/Runtime/Experimental/Chaos/Public/GeometryCollection/ManagedArrayTypeValues.inl
+++ b/Engine/Source/Runtime/Experimental/Chaos/Public/GeometryCollection/ManagedArrayTypeValues.inl
@@ -40,6 +40,7 @@ MANAGED_ARRAY_TYPE(Chaos::ThreadSafeSharedPtr_FImplicitObject, FImplicitObject3T
 MANAGED_ARRAY_TYPE(Chaos::NotThreadSafeSharedPtr_FImplicitObject, FImplicitObject3SharedPointer)
 MANAGED_ARRAY_TYPE(Chaos::FPBDRigidClusteredParticleHandle*, TPBDRigidClusteredParticleHandle3fPtr)
 MANAGED_ARRAY_TYPE(TUniquePtr<Chaos::FConvex>, FConvexUniquePtr)
+MANAGED_ARRAY_TYPE(TArray<FVector2D>, Vector2DArray)
 
 // NOTE: new types must be added at the bottom to keep serialization from breaking
 
diff --git a/Engine/Source/Runtime/Experimental/GeometryCollectionEngine/Private/GeometryCollection/GeometryCollectionComponent.cpp b/Engine/Source/Runtime/Experimental/GeometryCollectionEngine/Private/GeometryCollection/GeometryCollectionComponent.cpp
--- a/Engine/Source/Runtime/Experimental/GeometryCollectionEngine/Private/GeometryCollection/GeometryCollectionComponent.cpp
+++ b/Engine/Source/Runtime/Experimental/GeometryCollectionEngine/Private/GeometryCollection/GeometryCollectionComponent.cpp
@@ -1212,7 +1212,7 @@ void UGeometryCollectionComponent::InitConstantData(FGeometryCollectionConstantD
 		const TManagedArray<FVector3f>& TangentU = Collection->TangentU;
 		const TManagedArray<FVector3f>& TangentV = Collection->TangentV;
 		const TManagedArray<FVector3f>& Normal = Collection->Normal;
-		const TManagedArray<FVector2D>& UV = Collection->UV;
+		const TManagedArray<TArray<FVector2D>>& UVs = Collection->UVs;
 		const TManagedArray<FLinearColor>& Color = Collection->Color;
 		const TManagedArray<FLinearColor>& BoneColors = Collection->BoneColor;
 
@@ -1221,7 +1221,7 @@ void UGeometryCollectionComponent::InitConstantData(FGeometryCollectionConstantD
 		ConstantData->TangentU = TArray<FVector3f>(TangentU.GetData(), TangentU.Num());
 		ConstantData->TangentV = TArray<FVector3f>(TangentV.GetData(), TangentV.Num());
 		ConstantData->Normals = TArray<FVector3f>(Normal.GetData(), Normal.Num());
-		ConstantData->UVs = TArray<FVector2D>(UV.GetData(), UV.Num());
+		ConstantData->UVs = TArray<TArray<FVector2D>>(UVs.GetData(), UVs.Num());
 		ConstantData->Colors = TArray<FLinearColor>(Color.GetData(), Color.Num());
 
 		ConstantData->BoneColors.AddUninitialized(NumPoints);
@@ -2577,18 +2577,22 @@ void FScopedColorEdit::UpdateBoneColors()
 
 void UGeometryCollectionComponent::ApplyKinematicField(float Radius, FVector Position)
 {
-	FFieldSystemCommand Command = FFieldObjectCommands::CreateFieldCommand(EFieldPhysicsType::Field_DynamicState, new FRadialIntMask(Radius, Position, (int32)Chaos::EObjectStateType::Dynamic,
-			(int32)Chaos::EObjectStateType::Kinematic, ESetMaskConditionType::Field_Set_IFF_NOT_Interior));
-	DispatchFieldCommand(Command);
+	/*
+	FName TargetName = GetGeometryCollectionPhysicsTypeName(EGeometryCollectionPhysicsTypeEnum::Chaos_DynamicState);
+	DispatchFieldCommand({ TargetName,new FRadialIntMask(Radius, Position, (int32)Chaos::EObjectStateType::Dynamic,
+		(int32)Chaos::EObjectStateType::Kinematic, ESetMaskConditionType::Field_Set_IFF_NOT_Interior) });
+		*/
 }
 
 void UGeometryCollectionComponent::ApplyPhysicsField(bool Enabled, EGeometryCollectionPhysicsTypeEnum Target, UFieldSystemMetaData* MetaData, UFieldNodeBase* Field)
 {
+	/*
 	if (Enabled && Field)
 	{
-		FFieldSystemCommand Command = FFieldObjectCommands::CreateFieldCommand(GetGeometryCollectionPhysicsType(Target), Field, MetaData);  
+		FFieldSystemCommand Command = FFieldObjectCommands::CreateFieldCommand(GetGeometryCollectionPhysicsTypeName(Target), Field, MetaData);
 		DispatchFieldCommand(Command);
 	}
+	*/
 }
 
 void UGeometryCollectionComponent::DispatchFieldCommand(const FFieldSystemCommand& InCommand)
diff --git a/Engine/Source/Runtime/Experimental/GeometryCollectionEngine/Private/GeometryCollection/GeometryCollectionEngineUtility.cpp b/Engine/Source/Runtime/Experimental/GeometryCollectionEngine/Private/GeometryCollection/GeometryCollectionEngineUtility.cpp
--- a/Engine/Source/Runtime/Experimental/GeometryCollectionEngine/Private/GeometryCollection/GeometryCollectionEngineUtility.cpp
+++ b/Engine/Source/Runtime/Experimental/GeometryCollectionEngine/Private/GeometryCollection/GeometryCollectionEngineUtility.cpp
@@ -327,7 +327,8 @@ void GeometryCollectionEngineUtility::ComputeNormals(FGeometryCollection* Geomet
 	TArray<FVector2D> TmpUV;
 
 	TManagedArray<FIntVector>& Indices = GeometryCollection->Indices;
-	TManagedArray<FVector2D>& UV = GeometryCollection->UV;
+	const int32 UVChannelIndex = 0;
+	TManagedArray<TArray<FVector2D>>& UVs = GeometryCollection->UVs;
 	for (int i = 0; i < NumIndices; ++i)
 	{
 		FIntVector CurrTri = Indices[i];
@@ -336,9 +337,9 @@ void GeometryCollectionEngineUtility::ComputeNormals(FGeometryCollection* Geomet
 		TmpIndices.Add(CurrTri.Y);
 		TmpIndices.Add(CurrTri.Z);
 
-		TmpUV.Add(UV[CurrTri.X]);
-		TmpUV.Add(UV[CurrTri.Y]);
-		TmpUV.Add(UV[CurrTri.Z]);
+		TmpUV.Add(UVs[CurrTri.X][UVChannelIndex]);
+		TmpUV.Add(UVs[CurrTri.Y][UVChannelIndex]);
+		TmpUV.Add(UVs[CurrTri.Z][UVChannelIndex]);
 	}
 
 	// Make a copy of the vertex array
@@ -384,7 +385,8 @@ void GeometryCollectionEngineUtility::ComputeTangents(FGeometryCollection* Geome
 	TArray<FVector2D> TmpUV;
 
 	TManagedArray<FIntVector>& Indices = GeometryCollection->Indices;
-	TManagedArray<FVector2D>& UV = GeometryCollection->UV;
+	const int32 UVChannelIndex = 0;
+	TManagedArray<TArray<FVector2D>>& UVs = GeometryCollection->UVs;
 	for (int i = 0; i < NumIndices; ++i)
 	{
 		FIntVector CurrTri = Indices[i];
@@ -393,9 +395,9 @@ void GeometryCollectionEngineUtility::ComputeTangents(FGeometryCollection* Geome
 		TmpIndices.Add(CurrTri.Y);
 		TmpIndices.Add(CurrTri.Z);
 
-		TmpUV.Add(UV[CurrTri.X]);
-		TmpUV.Add(UV[CurrTri.Y]);
-		TmpUV.Add(UV[CurrTri.Z]);
+		TmpUV.Add(UVs[CurrTri.X][UVChannelIndex]);
+		TmpUV.Add(UVs[CurrTri.Y][UVChannelIndex]);
+		TmpUV.Add(UVs[CurrTri.Z][UVChannelIndex]);
 	}
 
 	// Make a copy of the vertex array
diff --git a/Engine/Source/Runtime/Experimental/GeometryCollectionEngine/Private/GeometryCollection/GeometryCollectionObject.cpp b/Engine/Source/Runtime/Experimental/GeometryCollectionEngine/Private/GeometryCollection/GeometryCollectionObject.cpp
--- a/Engine/Source/Runtime/Experimental/GeometryCollectionEngine/Private/GeometryCollection/GeometryCollectionObject.cpp
+++ b/Engine/Source/Runtime/Experimental/GeometryCollectionEngine/Private/GeometryCollection/GeometryCollectionObject.cpp
@@ -632,7 +632,7 @@ TUniquePtr<FGeometryCollectionNaniteData> UGeometryCollection::CreateNaniteData(
 
 	// Vertices Group
 	const TManagedArray<FVector3f>& VertexArray = Collection->Vertex;
-	const TManagedArray<FVector2D>& UVArray = Collection->UV;
+	const TManagedArray<TArray<FVector2D>>& UVsArray = Collection->UVs;
 	const TManagedArray<FLinearColor>& ColorArray = Collection->Color;
 	const TManagedArray<FVector3f>& TangentUArray = Collection->TangentU;
 	const TManagedArray<FVector3f>& TangentVArray = Collection->TangentV;
@@ -658,7 +658,7 @@ TUniquePtr<FGeometryCollectionNaniteData> UGeometryCollection::CreateNaniteData(
 	// Material Group
 	const int32 NumGeometry = Collection->NumElements(FGeometryCollection::GeometryGroup);
 
-	const uint32 NumTexCoords = 1;// NumTextureCoord;
+	const uint32 NumTexCoords = Collection->NumUVLayers();
 	const bool bHasColors = ColorArray.Num() > 0;
 
 	TArray<FStaticMeshBuildVertex> BuildVertices;
@@ -683,10 +683,13 @@ TUniquePtr<FGeometryCollectionNaniteData> UGeometryCollection::CreateNaniteData(
 			Vertex.TangentX = FVector3f::ZeroVector;
 			Vertex.TangentY = FVector3f::ZeroVector;
 			Vertex.TangentZ = NormalArray[VertexStart + VertexIndex];
-			Vertex.UVs[0] = UVArray[VertexStart + VertexIndex];
-			if (Vertex.UVs[0].ContainsNaN())
+			for (int32 UVIdx = 0; UVIdx < UVsArray[VertexStart + VertexIndex].Num(); ++UVIdx)
 			{
-				Vertex.UVs[0] = FVector2D::ZeroVector;
+				Vertex.UVs[UVIdx] = UVsArray[VertexStart + VertexIndex][UVIdx];
+				if (Vertex.UVs[UVIdx].ContainsNaN())
+				{
+					Vertex.UVs[UVIdx] = FVector2D::ZeroVector;
+				}
 			}
 		}
 
diff --git a/Engine/Source/Runtime/Experimental/GeometryCollectionEngine/Private/GeometryCollection/GeometryCollectionSceneProxy.cpp b/Engine/Source/Runtime/Experimental/GeometryCollectionEngine/Private/GeometryCollection/GeometryCollectionSceneProxy.cpp
--- a/Engine/Source/Runtime/Experimental/GeometryCollectionEngine/Private/GeometryCollection/GeometryCollectionSceneProxy.cpp
+++ b/Engine/Source/Runtime/Experimental/GeometryCollectionEngine/Private/GeometryCollection/GeometryCollectionSceneProxy.cpp
@@ -218,7 +218,7 @@ void FGeometryCollectionSceneProxy::InitResources()
 	
 	// Init buffers
 	VertexBuffers.PositionVertexBuffer.Init(NumVertices);
-	VertexBuffers.StaticMeshVertexBuffer.Init(NumVertices, 1);
+	VertexBuffers.StaticMeshVertexBuffer.Init(NumVertices, Chaos::GeometryCollection::MAX_NUM_UV_CHANNELS);
 	VertexBuffers.ColorVertexBuffer.Init(NumVertices);
 
 	// Init resources
@@ -320,12 +320,20 @@ void FGeometryCollectionSceneProxy::BuildGeometry( const FGeometryCollectionCons
 		OutVertices[PointIdx] =
 			FDynamicMeshVertex(
 				ConstantDataIn->Vertices[PointIdx],
-				ConstantDataIn->UVs[PointIdx],
+				ConstantDataIn->UVs[PointIdx][0],
 				bShowBoneColors||bEnableBoneSelection ?
 				ConstantDataIn->BoneColors[PointIdx].ToFColor(true) :
 				ConstantDataIn->Colors[PointIdx].ToFColor(true)
 			);
 		OutVertices[PointIdx].SetTangents(ConstantDataIn->TangentU[PointIdx], ConstantDataIn->TangentV[PointIdx], ConstantDataIn->Normals[PointIdx]);
+
+		if (ConstantDataIn->UVs[PointIdx].Num() > 1)
+		{
+			for (int32 UVLayerIdx = 1; UVLayerIdx < ConstantDataIn->UVs[PointIdx].Num(); ++UVLayerIdx)
+			{
+				OutVertices[PointIdx].TextureCoordinate[UVLayerIdx] = ConstantDataIn->UVs[PointIdx][UVLayerIdx];
+			}
+		}
 	});
 
 	check(ConstantDataIn->Indices.Num() * 3 == NumIndices);
@@ -380,7 +388,10 @@ void FGeometryCollectionSceneProxy::SetConstantData_RenderThread(FGeometryCollec
 
 			VertexBuffers.PositionVertexBuffer.VertexPosition(i) = Vertex.Position;
 			VertexBuffers.StaticMeshVertexBuffer.SetVertexTangents(i, Vertex.TangentX.ToFVector(), Vertex.GetTangentY(), Vertex.TangentZ.ToFVector());
-			VertexBuffers.StaticMeshVertexBuffer.SetVertexUV(i, 0, Vertex.TextureCoordinate[0]);
+			for (int UVChannelIndex = 0; UVChannelIndex < Chaos::GeometryCollection::MAX_NUM_UV_CHANNELS; UVChannelIndex++)
+			{
+				VertexBuffers.StaticMeshVertexBuffer.SetVertexUV(i, UVChannelIndex, Vertex.TextureCoordinate[UVChannelIndex]);
+			}
 			VertexBuffers.ColorVertexBuffer.VertexColor(i) = Vertex.Color;
 #if GEOMETRYCOLLECTION_EDITOR_SELECTION
 			if (bEnableBoneSelection)
diff --git a/Engine/Source/Runtime/Experimental/GeometryCollectionEngine/Private/GeometryCollection/GeometryCollectionSceneProxy.h b/Engine/Source/Runtime/Experimental/GeometryCollectionEngine/Private/GeometryCollection/GeometryCollectionSceneProxy.h
--- a/Engine/Source/Runtime/Experimental/GeometryCollectionEngine/Private/GeometryCollection/GeometryCollectionSceneProxy.h
+++ b/Engine/Source/Runtime/Experimental/GeometryCollectionEngine/Private/GeometryCollection/GeometryCollectionSceneProxy.h
@@ -103,7 +103,7 @@ struct FGeometryCollectionConstantData
 	TArray<FVector3f> Normals;
 	TArray<FVector3f> TangentU;
 	TArray<FVector3f> TangentV;
-	TArray<FVector2D> UVs;
+	TArray<TArray<FVector2D>> UVs;
 	TArray<FLinearColor> Colors;
 	TArray<int32> BoneMap;
 	TArray<FLinearColor> BoneColors;
diff --git a/Engine/Source/Runtime/Experimental/GeometryCollectionEngine/Public/GeometryCollection/GeometryCollectionComponent.h b/Engine/Source/Runtime/Experimental/GeometryCollectionEngine/Public/GeometryCollection/GeometryCollectionComponent.h
--- a/Engine/Source/Runtime/Experimental/GeometryCollectionEngine/Public/GeometryCollection/GeometryCollectionComponent.h
+++ b/Engine/Source/Runtime/Experimental/GeometryCollectionEngine/Public/GeometryCollection/GeometryCollectionComponent.h
@@ -392,7 +392,7 @@ public:
 
 	// Vertices Group
 	COPY_ON_WRITE_ATTRIBUTE(FVector3f, Vertex, FGeometryCollection::VerticesGroup) 	//GetVertexArray, GetVertexArrayCopyOnWrite, GetVertexArrayRest
-	COPY_ON_WRITE_ATTRIBUTE(FVector2D, UV, FGeometryCollection::VerticesGroup)		//GetUVArray
+	COPY_ON_WRITE_ATTRIBUTE(TArray<FVector2D>, UVs, FGeometryCollection::VerticesGroup)		//GetUVsArray
 	COPY_ON_WRITE_ATTRIBUTE(FLinearColor, Color, FGeometryCollection::VerticesGroup)//GetColorArray
 	COPY_ON_WRITE_ATTRIBUTE(FVector3f, TangentU, FGeometryCollection::VerticesGroup)	//GetTangentUArray
 	COPY_ON_WRITE_ATTRIBUTE(FVector3f, TangentV, FGeometryCollection::VerticesGroup)	//...
