Add support for targeting custom material IDs in fracture AutoUV tool
+ fix ambient occlusion bake to handle the default value being changed from 0 to 1

diff --git a/Engine/Plugins/Experimental/ChaosEditor/Source/FractureEditor/Private/FractureToolUV.cpp b/Engine/Plugins/Experimental/ChaosEditor/Source/FractureEditor/Private/FractureToolUV.cpp
--- a/Engine/Plugins/Experimental/ChaosEditor/Source/FractureEditor/Private/FractureToolUV.cpp
+++ b/Engine/Plugins/Experimental/ChaosEditor/Source/FractureEditor/Private/FractureToolUV.cpp
@@ -216,11 +216,14 @@ int32 UFractureToolAutoUV::ExecuteFracture(const FFractureToolContext& FractureC
 		FGeometryCollection& Collection = *FractureContext.GetGeometryCollection();
 		
 		int32 OutputRes = (int32)AutoUVSettings->Resolution;
+		TArray<int32> EmptyMaterialIDs;
 
 		if (AutoUVSettings->bDoUVLayout)
 		{
 			UVTask.EnterProgressFrame(1, LOCTEXT("LayOutUVIslands", "Laying out UV islands"));
-			if (!UE::PlanarCut::UVLayout(Collection, OutputRes, AutoUVSettings->GutterSize))
+			if (!UE::PlanarCut::UVLayout(Collection, OutputRes, AutoUVSettings->GutterSize,
+				AutoUVSettings->TargetMaterialIDs != ETargetMaterialIDs::SelectedIDs,
+				AutoUVSettings->TargetMaterialIDs == ETargetMaterialIDs::OddIDs ? EmptyMaterialIDs : AutoUVSettings->MaterialIDs))
 			{
 				// failed to do layout
 				return INDEX_NONE;
@@ -269,7 +272,9 @@ int32 UFractureToolAutoUV::ExecuteFracture(const FFractureToolContext& FractureC
 		AttribSettings.Curvature_ThicknessFactor = AutoUVSettings->ThicknessFactor;
 		AttribSettings.Curvature_MaxValue = AutoUVSettings->MaxCurvature;
 		AttribSettings.bNormalZ_TakeAbs = AutoUVSettings->bUseAbsoluteValue;
-		UE::PlanarCut::TextureInternalSurfaces(Collection, FMath::CeilToInt(AutoUVSettings->GutterSize), Attributes, AttribSettings, ImageBuilder);
+		UE::PlanarCut::TextureInternalSurfaces(Collection, FMath::CeilToInt(AutoUVSettings->GutterSize), Attributes, AttribSettings, ImageBuilder,
+			AutoUVSettings->TargetMaterialIDs != ETargetMaterialIDs::SelectedIDs, 
+			AutoUVSettings->TargetMaterialIDs == ETargetMaterialIDs::OddIDs ? EmptyMaterialIDs: AutoUVSettings->MaterialIDs);
 
 		UVTask.EnterProgressFrame(1, LOCTEXT("SavingTexture", "Saving result"));
 
diff --git a/Engine/Plugins/Experimental/ChaosEditor/Source/FractureEditor/Private/FractureToolUV.h b/Engine/Plugins/Experimental/ChaosEditor/Source/FractureEditor/Private/FractureToolUV.h
--- a/Engine/Plugins/Experimental/ChaosEditor/Source/FractureEditor/Private/FractureToolUV.h
+++ b/Engine/Plugins/Experimental/ChaosEditor/Source/FractureEditor/Private/FractureToolUV.h
@@ -30,6 +30,14 @@ enum class ETextureType
 	SpatialGradients
 };
 
+UENUM()
+enum class ETargetMaterialIDs
+{
+	OddIDs,
+	OddAndSelectedIDs,
+	SelectedIDs
+};
+
 /** Settings specifically related to the one-time destructive fracturing of a mesh **/
 UCLASS(config = EditorPerProjectUserSettings)
 class UFractureAutoUVSettings : public UFractureToolSettings
@@ -58,6 +66,14 @@ public:
 	UPROPERTY(VisibleAnywhere, Category = MapSettings)
 	UTexture2D* Result;
 
+	/** Choose whether to texture only odd material IDs (corresponding to internal faces) or a custom selection */
+	UPROPERTY(EditAnywhere, Category = MapSettings)
+	ETargetMaterialIDs TargetMaterialIDs;
+
+	/** Custom selection of material IDs to target for texturing */
+	UPROPERTY(EditAnywhere, Category = MapSettings, meta = (EditCondition = "TargetMaterialIDs != ETargetMaterialIDs::OddIDs", EditConditionHides))
+	TArray<int32> MaterialIDs;
+
 	/** Whether to prompt user for an asset name for each generated texture, or automatically place them next to the source geometry collections */
 	UPROPERTY(EditAnywhere, Category = MapSettings)
 	bool bPromptToSave = true;
diff --git a/Engine/Plugins/Experimental/PlanarCutPlugin/Source/PlanarCut/Private/FractureAutoUV.cpp b/Engine/Plugins/Experimental/PlanarCutPlugin/Source/PlanarCut/Private/FractureAutoUV.cpp
--- a/Engine/Plugins/Experimental/PlanarCutPlugin/Source/PlanarCut/Private/FractureAutoUV.cpp
+++ b/Engine/Plugins/Experimental/PlanarCutPlugin/Source/PlanarCut/Private/FractureAutoUV.cpp
@@ -116,16 +116,16 @@ namespace {
 			return false;
 		}
 
-		if (bOddMaterialAreInside && ((MaterialID % 2) == 0) == bActivateInsideTriangles)
+		if (bOddMaterialAreInside && ((MaterialID % 2) == 1) == bActivateInsideTriangles)
 		{
-			return false;
+			return true;
 		}
-		if (InsideMaterials.Num() > 0 && (InsideMaterials.Contains(MaterialID) != bActivateInsideTriangles))
+		if (InsideMaterials.Num() > 0 && (InsideMaterials.Contains(MaterialID) == bActivateInsideTriangles))
 		{
-			return false;
+			return true;
 		}
 
-		return true;
+		return false;
 	}
 }
 
@@ -484,20 +484,21 @@ void TextureInternalSurfaces(
 	TArrayView<int32> WhichMaterials
 )
 {
-	TArray<bool> InsideTriangles;
-	int32 NumInsideTris = SetActiveTriangles(&Collection, InsideTriangles, true, bOnlyOddMaterials, WhichMaterials);
-	FGeomFlatUVMesh UVMesh(&Collection, InsideTriangles, NumInsideTris);
+	TArray<bool> ToTextureTriangles;
+	int32 NumTextureTris = SetActiveTriangles(&Collection, ToTextureTriangles, true, bOnlyOddMaterials, WhichMaterials);
+	FGeomFlatUVMesh UVMesh(&Collection, ToTextureTriangles, NumTextureTris);
 
 	FImageOccupancyMap OccupancyMap;
 	OccupancyMap.GutterSize = GutterSize;
 	OccupancyMap.Initialize(TextureOut.GetDimensions());
 	OccupancyMap.ComputeFromUVSpaceMesh(UVMesh, [](int32 TriangleID) { return TriangleID; });
 
-	FGeomMesh InsideMesh(&Collection, InsideTriangles, NumInsideTris);
+	FGeomMesh ToTextureMesh(&Collection, ToTextureTriangles, NumTextureTris);
 
+	// Outside triangles are identified by odd material ID
 	TArray<bool> OutsideTriangles;
-	int32 NumOutsideTris = SetActiveTriangles(&Collection, OutsideTriangles, false, bOnlyOddMaterials, WhichMaterials);
-	FGeomMesh OutsideMesh(InsideMesh, OutsideTriangles, NumOutsideTris);
+	int32 NumOutsideTris = SetActiveTriangles(&Collection, OutsideTriangles, false, true, {});
+	FGeomMesh OutsideMesh(ToTextureMesh, OutsideTriangles, NumOutsideTris);
 	TMeshAABBTree3<FGeomMesh> OutsideSpatial(&OutsideMesh);
 
 	int DistanceToExternalIdx = BakeAttributes.IndexOf((int)EBakeAttributes::DistanceToExternal);
@@ -523,10 +524,10 @@ void TextureInternalSurfaces(
 	FDynamicMeshCollection CollectionMeshes(&Collection, TransformIndices, FTransform::Identity, false);
 	if (bNeedsDynamicMeshes)
 	{
-		TSet<int32> InsideMaterials;
+		TSet<int32> TargetMaterials;
 		for (int32 MatID : WhichMaterials)
 		{
-			InsideMaterials.Add(MatID);
+			TargetMaterials.Add(MatID);
 		}
 
 		// To bake only internal faces, unset texture on everything else
@@ -539,10 +540,11 @@ void TextureInternalSurfaces(
 			FDynamicMeshUVOverlay* UV = Mesh.Attributes()->PrimaryUV();
 			for (int TID : Mesh.TriangleIndicesItr())
 			{
-				bool bIsInsideTri = IsTriActive(
+				// TODO: separate the concept of inside from the concept of needs texture!!
+				bool bIsTextureTri = IsTriActive(
 					AugmentedDynamicMesh::GetVisibility(Mesh, TID), 
-					Mesh.Attributes()->GetMaterialID()->GetValue(TID), InsideMaterials, true, bOnlyOddMaterials);
-				if (!bIsInsideTri)
+					Mesh.Attributes()->GetMaterialID()->GetValue(TID), TargetMaterials, true, bOnlyOddMaterials);
+				if (!bIsTextureTri)
 				{
 					UV->UnsetTriangle(TID);
 				}
@@ -630,7 +632,7 @@ void TextureInternalSurfaces(
 	{
 		FImageDimensions Dimensions = TextureOut.GetDimensions();
 		int64 TexelsNum = Dimensions.Num();
-		AmbientValues.SetNumZeroed(TexelsNum);
+		AmbientValues.Init(1.0f, TexelsNum);
 		for (int MeshIdx = 0; MeshIdx < CollectionMeshes.Meshes.Num(); MeshIdx++)
 		{
 			int32 TransformIdx = CollectionMeshes.Meshes[MeshIdx].TransformIndex;
@@ -645,8 +647,8 @@ void TextureInternalSurfaces(
 			BakeCache.SetBakeTargetMesh(&Mesh);
 			BakeCache.SetDimensions(Dimensions);
 			BakeCache.SetUVLayer(0);
-			// thickness is used for raycasting correspondences between detail and target mesh
-			// TODO: make cache aware that detail mesh == target mesh so it doesn't have to do this?
+			BakeCache.SetCorrespondenceStrategy(FMeshImageBakingCache::ECorrespondenceStrategy::Identity);
+			// thickness shouldn't matter because we're using Identity ECorrespondenceStrategy::Identity
 			BakeCache.SetThickness(KINDA_SMALL_NUMBER);
 			BakeCache.ValidateCache();
 
@@ -666,13 +668,13 @@ void TextureInternalSurfaces(
 			for (int64 LinearIdx = 0; LinearIdx < TexelsNum; LinearIdx++)
 			{
 				float& AmbientOut = AmbientValues[LinearIdx];
-				AmbientOut = FMathf::Max(AmbientOut, Result->GetPixel(LinearIdx).X);
+				AmbientOut = FMathf::Min(AmbientOut, Result->GetPixel(LinearIdx).X);
 			}
 		}
 	}
 
 	ParallelFor(OccupancyMap.Dimensions.GetHeight(),
-		[&AttributeSettings, &TextureOut, &UVMesh, &OccupancyMap, &InsideMesh, &OutsideSpatial,
+		[&AttributeSettings, &TextureOut, &UVMesh, &OccupancyMap, &ToTextureMesh, &OutsideSpatial,
 		 &DistanceToExternalIdx, &AmbientIdx, &NormalZIdx, &PosXIdx, &PosYIdx, &PosZIdx](int32 Y)
 	{
 		for (int32 X = 0; X < OccupancyMap.Dimensions.GetWidth(); X++)
@@ -692,13 +694,13 @@ void TextureInternalSurfaces(
 				bool bNeedsNormal = NormalZIdx > -1 || AmbientIdx > -1;
 				if (bNeedsNormal)
 				{
-					Normal = InsideMesh.GetInterpolatedNormal(TID, Bary);
+					Normal = ToTextureMesh.GetInterpolatedNormal(TID, Bary);
 				}
 
 				if (DistanceToExternalIdx > -1 || PosXIdx > -1 || PosYIdx > -1 || PosZIdx > -1)
 				{
 					FTriangle3d Tri;
-					InsideMesh.GetTriVertices(TID, Tri.V[0], Tri.V[1], Tri.V[2]);
+					ToTextureMesh.GetTriVertices(TID, Tri.V[0], Tri.V[1], Tri.V[2]);
 					FVector3d InsidePoint = Tri.BarycentricPoint(Bary);
 
 					if (DistanceToExternalIdx > -1)
diff --git a/Engine/Plugins/Experimental/PlanarCutPlugin/Source/PlanarCut/Public/FractureAutoUV.h b/Engine/Plugins/Experimental/PlanarCutPlugin/Source/PlanarCut/Public/FractureAutoUV.h
--- a/Engine/Plugins/Experimental/PlanarCutPlugin/Source/PlanarCut/Public/FractureAutoUV.h
+++ b/Engine/Plugins/Experimental/PlanarCutPlugin/Source/PlanarCut/Public/FractureAutoUV.h
@@ -14,15 +14,15 @@ namespace UE { namespace PlanarCut {
  * @param Collection		The collection to be atlas'd
  * @param UVRes				Target resolution for the atlas
  * @param GutterSize		Space to leave between UV islands, in pixels at the target resolution
- * @param bOnlyOddMaterials	If true, restrict UV island layout to odd-numbered material IDs
- * @param WhichMaterials	If non-empty, restrict UV island layout to only the listed material IDs
+ * @param bOddMaterials		If true, consider odd-numbered material IDs for UV island layout
+ * @param WhichMaterials	If non-empty, consider listed material IDs for UV island layout
  * @param bRecreateUVsForDegenerateIslands If true, detect and fix islands that don't have proper UVs (i.e. UVs all zero or otherwise collapsed to a point)
  */
 bool PLANARCUT_API UVLayout(
 	FGeometryCollection& Collection,
 	int32 UVRes = 1024,
 	float GutterSize = 1,
-	bool bOnlyOddMaterials = true,
+	bool bOddMaterials = true,
 	TArrayView<int32> WhichMaterials = TArrayView<int32>(),
 	bool bRecreateUVsForDegenerateIslands = true
 );
@@ -71,8 +71,8 @@ struct FTextureAttributeSettings
  * @param BakeAttributes	Which attributes to bake into which color channel
  * @param AttributeSettings	Settings for the BakeAttributes
  * @param TextureOut		Texture to write to
- * @param bOnlyOddMaterials	If true, restrict UV island layout to odd-numbered material IDs
- * @param WhichMaterials	If non-empty, restrict UV island layout to only the listed material IDs
+ * @param bOddMaterials		If true, apply texture to odd-numbered material IDs
+ * @param WhichMaterials	If non-empty, apply texture to the listed material IDs
  */
 void PLANARCUT_API TextureInternalSurfaces(
 	FGeometryCollection& Collection,
@@ -80,7 +80,7 @@ void PLANARCUT_API TextureInternalSurfaces(
 	UE::Geometry::FIndex4i BakeAttributes,
 	const FTextureAttributeSettings& AttributeSettings,
 	UE::Geometry::TImageBuilder<FVector4f>& TextureOut,
-	bool bOnlyOddMaterials = true,
+	bool bOddMaterials = true,
 	TArrayView<int32> WhichMaterials = TArrayView<int32>()
 );
 
