make fracture UVs try to align to major axes

#rnx

diff --git a/Engine/Plugins/Experimental/PlanarCutPlugin/Source/PlanarCut/Private/GeometryMeshConversion.cpp b/Engine/Plugins/Experimental/PlanarCutPlugin/Source/PlanarCut/Private/GeometryMeshConversion.cpp
--- a/Engine/Plugins/Experimental/PlanarCutPlugin/Source/PlanarCut/Private/GeometryMeshConversion.cpp
+++ b/Engine/Plugins/Experimental/PlanarCutPlugin/Source/PlanarCut/Private/GeometryMeshConversion.cpp
@@ -450,6 +450,24 @@ namespace AugmentedDynamicMesh
 			}
 		}
 	}
+} // namespace AugmentedDynamicMesh
+
+
+namespace
+{
+	/// Make a frame w/ Z aligned to the given normal
+	/// and X aligned to the major axis most perpendicular to the normal
+	FFrame3d AxisAlignedFrame(const FPlane& Plane)
+	{
+		FVector3d Origin = (FVector3d)Plane.GetOrigin();
+		FVector3d Normal = (FVector3d)Plane.GetNormal();
+		int32 MinIdx = VectorUtil::MinAbsElementIndex(Normal);
+		FFrame3d Frame(Origin, Normal);
+		FVector3d Target(0, 0, 0);
+		Target[MinIdx] = 1;
+		Frame.ConstrainedAlignAxis(0, Target, Normal);
+		return Frame;
+	}
 }
 
 
@@ -774,7 +792,7 @@ void FCellMeshes::CreateMeshesForBoundedPlanesWithoutNoise(int NumCells, const F
 
 		const TArray<int>& PlaneBoundary = Cells.PlaneBoundaries[PlaneIdx];
 		FVector3f Normal(Cells.Planes[PlaneIdx].GetNormal());
-		FFrame3d PlaneFrame(Cells.Planes[PlaneIdx]);
+		FFrame3d PlaneFrame = AxisAlignedFrame(Cells.Planes[PlaneIdx]);
 		FVertexInfo PlaneVertInfo;
 		PlaneVertInfo.bHaveC = true;
 		PlaneVertInfo.bHaveUV = false;
@@ -1185,7 +1203,7 @@ void FCellMeshes::CreateMeshesForBoundedPlanesWithNoise(int NumCells, const FPla
 	TArray<FFrame3d> PlaneFrames; PlaneFrames.Reserve(Cells.Planes.Num());
 	for (int PlaneIdx = 0; PlaneIdx < Cells.Planes.Num(); PlaneIdx++)
 	{
-		PlaneFrames.Emplace(Cells.Planes[PlaneIdx]);
+		PlaneFrames.Add(AxisAlignedFrame(Cells.Planes[PlaneIdx]));
 	}
 	// first pass to compute min UV for each plane
 	for (FCellInfo& CellInfo : CellMeshes)
@@ -1242,7 +1260,7 @@ void FCellMeshes::CreateMeshesForSinglePlane(const FPlanarCells& Cells, const FA
 	int MID = PlaneToMaterial(0);
 	FPlane Plane = Cells.Planes[0];
 
-	FFrame3d PlaneFrame(Plane);
+	FFrame3d PlaneFrame = AxisAlignedFrame(Plane);
 	FInterval1d ZRange;
 	FAxisAlignedBox2d XYRange;
 	for (int CornerIdx = 0; CornerIdx < 8; CornerIdx++)
diff --git a/Engine/Source/Runtime/GeometryCore/Public/VectorTypes.h b/Engine/Source/Runtime/GeometryCore/Public/VectorTypes.h
--- a/Engine/Source/Runtime/GeometryCore/Public/VectorTypes.h
+++ b/Engine/Source/Runtime/GeometryCore/Public/VectorTypes.h
@@ -772,7 +772,7 @@ constexpr T MaxAbsElement(const UE::Math::TVector<T>& Vector)
 
 /** @return 0/1/2 index of maximum absolute-value element */
 template<typename T>
-constexpr T MaxAbsElementIndex(const UE::Math::TVector<T>& Vector)
+constexpr int32 MaxAbsElementIndex(const UE::Math::TVector<T>& Vector)
 {
 	return TMathUtil<T>::Max3Index(TMathUtil<T>::Abs(Vector.X), TMathUtil<T>::Abs(Vector.Y), TMathUtil<T>::Abs(Vector.Z));
 }
@@ -785,7 +785,7 @@ constexpr T MinAbsElement(const UE::Math::TVector<T>& Vector)
 
 /** @return 0/1/2 index of minimum absolute-value element */
 template<typename T>
-constexpr T MinAbsElementIndex(const UE::Math::TVector<T>& Vector)
+constexpr int32 MinAbsElementIndex(const UE::Math::TVector<T>& Vector)
 {
 	return TMathUtil<T>::Min3Index(TMathUtil<T>::Abs(Vector.X), TMathUtil<T>::Abs(Vector.Y), TMathUtil<T>::Abs(Vector.Z));
 }
