uv layer selector for fracture autouv

#rnx

diff --git a/Engine/Plugins/Experimental/ChaosEditor/Source/FractureEditor/Private/FractureToolUV.cpp b/Engine/Plugins/Experimental/ChaosEditor/Source/FractureEditor/Private/FractureToolUV.cpp
--- a/Engine/Plugins/Experimental/ChaosEditor/Source/FractureEditor/Private/FractureToolUV.cpp
+++ b/Engine/Plugins/Experimental/ChaosEditor/Source/FractureEditor/Private/FractureToolUV.cpp
@@ -24,6 +24,32 @@
 
 using namespace UE::Geometry;
 
+
+void UFractureAutoUVSettings::SetNumUVChannels(int32 NumUVChannels)
+{
+	NumUVChannels = FMath::Clamp(NumUVChannels, 1, GeometryCollectionUV::MAX_NUM_UV_CHANNELS);
+	UVChannelNamesList.Reset();
+	for (int32 k = 0; k < NumUVChannels; ++k)
+	{
+		UVChannelNamesList.Add(FString::Printf(TEXT("UV %d"), k));
+	}
+	if (GetSelectedChannelIndex(false) == INDEX_NONE)
+	{
+		UVChannel = UVChannelNamesList[0];
+	}
+}
+
+int32 UFractureAutoUVSettings::GetSelectedChannelIndex(bool bForceToZeroOnFailure)
+{
+	int32 FoundIndex = UVChannelNamesList.IndexOfByKey(UVChannel);
+	if (FoundIndex == INDEX_NONE && bForceToZeroOnFailure)
+	{
+		FoundIndex = 0;
+	}
+	return FoundIndex;
+}
+
+
 UFractureToolAutoUV::UFractureToolAutoUV(const FObjectInitializer& ObjInit)
 	: Super(ObjInit)
 {
@@ -71,6 +97,18 @@ TArray<UObject*> UFractureToolAutoUV::GetSettingsObjects() const
 
 void UFractureToolAutoUV::FractureContextChanged()
 {
+	int32 MinUVChannels = GeometryCollectionUV::MAX_NUM_UV_CHANNELS;
+
+	TSet<UGeometryCollectionComponent*> GeomCompSelection;
+	GetSelectedGeometryCollectionComponents(GeomCompSelection);
+	for (UGeometryCollectionComponent* GeometryCollectionComponent : GeomCompSelection)
+	{
+		int32 NumChannels = GeometryCollectionComponent->GetRestCollection()->GetGeometryCollection()->NumUVLayers();
+		MinUVChannels = FMath::Min(MinUVChannels, NumChannels);
+	}
+
+	AutoUVSettings->SetNumUVChannels(MinUVChannels);
+
 }
 
 void UFractureToolAutoUV::Render(const FSceneView* View, FViewport* Viewport, FPrimitiveDrawInterface* PDI)
@@ -218,10 +256,12 @@ int32 UFractureToolAutoUV::ExecuteFracture(const FFractureToolContext& FractureC
 		int32 OutputRes = (int32)AutoUVSettings->Resolution;
 		TArray<int32> EmptyMaterialIDs;
 
+		int32 UVLayer = AutoUVSettings->GetSelectedChannelIndex();
+
 		if (AutoUVSettings->bDoUVLayout)
 		{
 			UVTask.EnterProgressFrame(1, LOCTEXT("LayOutUVIslands", "Laying out UV islands"));
-			if (!UE::PlanarCut::UVLayout(Collection, OutputRes, AutoUVSettings->GutterSize,
+			if (!UE::PlanarCut::UVLayout(UVLayer, Collection, OutputRes, AutoUVSettings->GutterSize,
 				AutoUVSettings->TargetMaterialIDs != ETargetMaterialIDs::SelectedIDs,
 				AutoUVSettings->TargetMaterialIDs == ETargetMaterialIDs::OddIDs ? EmptyMaterialIDs : AutoUVSettings->MaterialIDs))
 			{
@@ -272,7 +312,7 @@ int32 UFractureToolAutoUV::ExecuteFracture(const FFractureToolContext& FractureC
 		AttribSettings.Curvature_ThicknessFactor = AutoUVSettings->ThicknessFactor;
 		AttribSettings.Curvature_MaxValue = AutoUVSettings->MaxCurvature;
 		AttribSettings.bNormalZ_TakeAbs = AutoUVSettings->bUseAbsoluteValue;
-		UE::PlanarCut::TextureInternalSurfaces(Collection, FMath::CeilToInt(AutoUVSettings->GutterSize), Attributes, AttribSettings, ImageBuilder,
+		UE::PlanarCut::TextureInternalSurfaces(UVLayer, Collection, FMath::CeilToInt(AutoUVSettings->GutterSize), Attributes, AttribSettings, ImageBuilder,
 			AutoUVSettings->TargetMaterialIDs != ETargetMaterialIDs::SelectedIDs, 
 			AutoUVSettings->TargetMaterialIDs == ETargetMaterialIDs::OddIDs ? EmptyMaterialIDs: AutoUVSettings->MaterialIDs);
 
diff --git a/Engine/Plugins/Experimental/ChaosEditor/Source/FractureEditor/Private/FractureToolUV.h b/Engine/Plugins/Experimental/ChaosEditor/Source/FractureEditor/Private/FractureToolUV.h
--- a/Engine/Plugins/Experimental/ChaosEditor/Source/FractureEditor/Private/FractureToolUV.h
+++ b/Engine/Plugins/Experimental/ChaosEditor/Source/FractureEditor/Private/FractureToolUV.h
@@ -3,6 +3,7 @@
 #pragma once
 
 #include "FractureTool.h"
+#include "GeometryCollection/GeometryCollection.h"
 
 #include "FractureToolUV.generated.h"
 
@@ -48,10 +49,27 @@ public:
 
 	UFractureAutoUVSettings(const FObjectInitializer& ObjInit)
 		: Super(ObjInit)
-	{}
+	{
+		SetNumUVChannels(1);
+	}
+
+	UPROPERTY(EditAnywhere, Category = UVChannel, meta = (DisplayName = "UV Channel", GetOptions = GetUVChannelNamesFunc))
+	FString UVChannel;
+
+	UFUNCTION()
+	const TArray<FString>& GetUVChannelNamesFunc()
+	{
+		return UVChannelNamesList;
+	}
+
+	UPROPERTY(meta = (TransientToolProperty))
+	TArray<FString> UVChannelNamesList;
+
+	void SetNumUVChannels(int32 NumUVChannels);
+	int32 GetSelectedChannelIndex(bool bForceToZeroOnFailure = true);
 
 	/** Whether to layout UVs (to ensure unique texels for each face), or just use the existing UVs */
-	UPROPERTY(EditAnywhere, Category = Atlas)
+	UPROPERTY(EditAnywhere, Category = Atlas, meta = (DisplayName = "Do UV Layout"))
 	bool bDoUVLayout = true;
 
 	/** The pixel resolution of the generated map */
diff --git a/Engine/Plugins/Experimental/PlanarCutPlugin/Source/PlanarCut/Private/FractureAutoUV.cpp b/Engine/Plugins/Experimental/PlanarCutPlugin/Source/PlanarCut/Private/FractureAutoUV.cpp
--- a/Engine/Plugins/Experimental/PlanarCutPlugin/Source/PlanarCut/Private/FractureAutoUV.cpp
+++ b/Engine/Plugins/Experimental/PlanarCutPlugin/Source/PlanarCut/Private/FractureAutoUV.cpp
@@ -167,21 +167,33 @@ struct FGeomMesh : public UE::Geometry::FUVPacker::IUVMeshView
 	int32 NumTriangles;
 	TArray<FVector3d> GlobalVertices; // vertices transformed to global space
 	TArray<FVector3f> GlobalNormals; // normals transformed to global space
+	int32 UVLayer = 0;
 
 	// Construct a mesh from a geometry collection and a mask of active triangles
-	FGeomMesh(FGeometryCollection* Collection, TArrayView<bool> ActiveTriangles, int32 NumTriangles) 
-		: Collection(Collection), ActiveTriangles(ActiveTriangles), NumTriangles(NumTriangles)
+	FGeomMesh(int32 UVLayer, FGeometryCollection* Collection, TArrayView<bool> ActiveTriangles, int32 NumTriangles) 
+		: Collection(Collection), ActiveTriangles(ActiveTriangles), NumTriangles(NumTriangles), UVLayer(UVLayer)
 	{
+		ValidateUVLayer();
 		InitVertices();
 	}
 
 	// Construct a mesh from an existing FGeomMesh and a mask of active triangles
 	FGeomMesh(const FGeomMesh& OtherMesh, TArrayView<bool> ActiveTriangles, int32 NumTriangles)
-		: Collection(OtherMesh.Collection), ActiveTriangles(ActiveTriangles), NumTriangles(NumTriangles)
+		: Collection(OtherMesh.Collection), ActiveTriangles(ActiveTriangles), NumTriangles(NumTriangles),
+		  UVLayer(OtherMesh.UVLayer)
 	{
+		ValidateUVLayer();
 		GlobalVertices = OtherMesh.GlobalVertices;
 	}
 
+	void ValidateUVLayer()
+	{
+		if (!ensure(UVLayer >= 0 && UVLayer < Collection->NumUVLayers()))
+		{
+			UVLayer = FMath::Clamp(UVLayer, 0, Collection->NumUVLayers());
+		}
+	}
+
 	void InitVertices()
 	{
 		// transform all vertices from local to global space
@@ -217,12 +229,12 @@ struct FGeomMesh : public UE::Geometry::FUVPacker::IUVMeshView
 
 	virtual FVector2f GetUV(int32 EID) const
 	{
-		return FVector2f(Collection->UVs[EID][0]);
+		return FVector2f(Collection->UVs[EID][UVLayer]);
 	}
 
 	virtual void SetUV(int32 EID, FVector2f UVIn)
 	{
-		FVector2D& UV = Collection->UVs[EID][0];
+		FVector2D& UV = Collection->UVs[EID][UVLayer];
 		UV.X = UVIn.X;
 		UV.Y = UVIn.Y;
 	}
@@ -287,10 +299,21 @@ struct FGeomFlatUVMesh
 	FGeometryCollection* Collection;
 	const TArrayView<bool> ActiveTriangles;
 	int32 NumTriangles;
+	int32 UVLayer;
+
+	FGeomFlatUVMesh(int32 UVLayer, FGeometryCollection* Collection, TArrayView<bool> ActiveTriangles, int32 NumTriangles)
+		: Collection(Collection), ActiveTriangles(ActiveTriangles), NumTriangles(NumTriangles), UVLayer(UVLayer)
+	{
+		ValidateUVLayer();
+	}
 
-	FGeomFlatUVMesh(FGeometryCollection* Collection, TArrayView<bool> ActiveTriangles, int32 NumTriangles)
-		: Collection(Collection), ActiveTriangles(ActiveTriangles), NumTriangles(NumTriangles)
-	{}
+	void ValidateUVLayer()
+	{
+		if (!ensure(UVLayer >= 0 && UVLayer < Collection->NumUVLayers()))
+		{
+			UVLayer = FMath::Clamp(UVLayer, 0, Collection->NumUVLayers());
+		}
+	}
 
 	inline FIndex3i GetTriangle(int32 TID) const
 	{
@@ -299,7 +322,7 @@ struct FGeomFlatUVMesh
 
 	inline FVector3d GetVertex(int32 VID) const
 	{
-		const FVector2D& UV = Collection->UVs[VID][0];
+		const FVector2D& UV = Collection->UVs[VID][UVLayer];
 		return FVector3d(UV.X, UV.Y, 0);
 	}
 
@@ -347,6 +370,7 @@ struct FGeomFlatUVMesh
 
 
 bool UVLayout(
+	int32 TargetUVLayer,
 	FGeometryCollection& Collection,
 	int32 UVRes,
 	float GutterSize,
@@ -357,7 +381,7 @@ bool UVLayout(
 {
 	TArray<bool> ActiveTriangles;
 	int32 NumActive = SetActiveTriangles(&Collection, ActiveTriangles, true, bOnlyOddMaterials, WhichMaterials);
-	FGeomMesh UVMesh(&Collection, ActiveTriangles, NumActive);
+	FGeomMesh UVMesh(TargetUVLayer, &Collection, ActiveTriangles, NumActive);
 
 	TArray<TArray<int32>> UVIslands;
 	UE::UVPacking::CreateUVIslandsFromMeshTopology<FGeomMesh>(UVMesh, UVIslands);
@@ -475,6 +499,7 @@ bool UVLayout(
 
 
 void TextureInternalSurfaces(
+	int32 TargetUVLayer,
 	FGeometryCollection& Collection,
 	int32 GutterSize,
 	FIndex4i BakeAttributes,
@@ -486,14 +511,14 @@ void TextureInternalSurfaces(
 {
 	TArray<bool> ToTextureTriangles;
 	int32 NumTextureTris = SetActiveTriangles(&Collection, ToTextureTriangles, true, bOnlyOddMaterials, WhichMaterials);
-	FGeomFlatUVMesh UVMesh(&Collection, ToTextureTriangles, NumTextureTris);
+	FGeomFlatUVMesh UVMesh(TargetUVLayer, &Collection, ToTextureTriangles, NumTextureTris);
 
 	FImageOccupancyMap OccupancyMap;
 	OccupancyMap.GutterSize = GutterSize;
 	OccupancyMap.Initialize(TextureOut.GetDimensions());
 	OccupancyMap.ComputeFromUVSpaceMesh(UVMesh, [](int32 TriangleID) { return TriangleID; });
 
-	FGeomMesh ToTextureMesh(&Collection, ToTextureTriangles, NumTextureTris);
+	FGeomMesh ToTextureMesh(TargetUVLayer, &Collection, ToTextureTriangles, NumTextureTris);
 
 	// Outside triangles are identified by odd material ID
 	TArray<bool> OutsideTriangles;
@@ -536,7 +561,7 @@ void TextureInternalSurfaces(
 			int32 TransformIdx = CollectionMeshes.Meshes[MeshIdx].TransformIndex;
 			FDynamicMesh3& Mesh = CollectionMeshes.Meshes[MeshIdx].AugMesh;
 			FMeshNormals::InitializeOverlayToPerVertexNormals(Mesh.Attributes()->PrimaryNormals(), true);
-			AugmentedDynamicMesh::InitializeOverlayToPerVertexUVs(Mesh, Collection.NumUVLayers());
+			AugmentedDynamicMesh::InitializeOverlayToPerVertexUVs(Mesh, 1, TargetUVLayer); // build an overlay w/ *only* the uv layer we care about
 			FDynamicMeshUVOverlay* UV = Mesh.Attributes()->PrimaryUV();
 			for (int TID : Mesh.TriangleIndicesItr())
 			{
@@ -564,7 +589,6 @@ void TextureInternalSurfaces(
 
 			int32 TransformIdx = CollectionMeshes.Meshes[MeshIdx].TransformIndex;
 			FDynamicMesh3& Mesh = CollectionMeshes.Meshes[MeshIdx].AugMesh;
-			FDynamicMeshUVOverlay* UV = Mesh.Attributes()->PrimaryUV();
 
 			FDynamicMeshAABBTree3 Spatial(&Mesh);
 			TFastWindingTree<FDynamicMesh3> FastWinding(&Spatial);
@@ -637,7 +661,6 @@ void TextureInternalSurfaces(
 		{
 			int32 TransformIdx = CollectionMeshes.Meshes[MeshIdx].TransformIndex;
 			FDynamicMesh3& Mesh = CollectionMeshes.Meshes[MeshIdx].AugMesh;
-			FDynamicMeshUVOverlay* UV = Mesh.Attributes()->PrimaryUV();
 
 			FDynamicMeshAABBTree3 Spatial(&Mesh);
 
diff --git a/Engine/Plugins/Experimental/PlanarCutPlugin/Source/PlanarCut/Private/GeometryMeshConversion.cpp b/Engine/Plugins/Experimental/PlanarCutPlugin/Source/PlanarCut/Private/GeometryMeshConversion.cpp
--- a/Engine/Plugins/Experimental/PlanarCutPlugin/Source/PlanarCut/Private/GeometryMeshConversion.cpp
+++ b/Engine/Plugins/Experimental/PlanarCutPlugin/Source/PlanarCut/Private/GeometryMeshConversion.cpp
@@ -235,7 +235,7 @@ namespace AugmentedDynamicMesh
 		Vs->GetValue(VID, V);
 	}
 
-	void InitializeOverlayToPerVertexUVs(FDynamicMesh3& Mesh, int32 NumUVLayers)
+	void InitializeOverlayToPerVertexUVs(FDynamicMesh3& Mesh, int32 NumUVLayers, int32 FirstUVLayer)
 	{
 		Mesh.Attributes()->SetNumUVLayers(NumUVLayers);
 		for (int UVLayer = 0; UVLayer < NumUVLayers; ++UVLayer)
@@ -247,7 +247,7 @@ namespace AugmentedDynamicMesh
 			for (int VID : Mesh.VertexIndicesItr())
 			{
 				FVector2f UV;
-				GetUV(Mesh, VID, UV, UVLayer);
+				GetUV(Mesh, VID, UV, UVLayer + FirstUVLayer);
 				int UVID = UVs->AppendElement(UV);
 				VertToUVMap[VID] = UVID;
 			}
diff --git a/Engine/Plugins/Experimental/PlanarCutPlugin/Source/PlanarCut/Private/GeometryMeshConversion.h b/Engine/Plugins/Experimental/PlanarCutPlugin/Source/PlanarCut/Private/GeometryMeshConversion.h
--- a/Engine/Plugins/Experimental/PlanarCutPlugin/Source/PlanarCut/Private/GeometryMeshConversion.h
+++ b/Engine/Plugins/Experimental/PlanarCutPlugin/Source/PlanarCut/Private/GeometryMeshConversion.h
@@ -34,7 +34,8 @@ namespace AugmentedDynamicMesh
 	void PLANARCUT_API GetUV(const UE::Geometry::FDynamicMesh3& Mesh, int VID, FVector2f& UV, int UVLayer);
 	void PLANARCUT_API SetTangent(UE::Geometry::FDynamicMesh3& Mesh, int VID, FVector3f Normal, FVector3f TangentU, FVector3f TangentV);
 	void PLANARCUT_API GetTangent(const UE::Geometry::FDynamicMesh3& Mesh, int VID, FVector3f& U, FVector3f& V);
-	void PLANARCUT_API InitializeOverlayToPerVertexUVs(UE::Geometry::FDynamicMesh3& Mesh, int32 NumUVLayers);
+	/// Initialize UV overlays based on the custom AugmentedDynamicMesh per-vertex UV attributes.  Optionally use FirstUVLayer to skip layers
+	void PLANARCUT_API InitializeOverlayToPerVertexUVs(UE::Geometry::FDynamicMesh3& Mesh, int32 NumUVLayers, int32 FirstUVLayer = 0);
 	void PLANARCUT_API InitializeOverlayToPerVertexTangents(UE::Geometry::FDynamicMesh3& Mesh);
 	void PLANARCUT_API ComputeTangents(UE::Geometry::FDynamicMesh3& Mesh, bool bOnlyOddMaterials, const TArrayView<const int32>& WhichMaterials, bool bRecomputeNormals = true);
 	void PLANARCUT_API AddCollisionSamplesPerComponent(UE::Geometry::FDynamicMesh3& Mesh, double Spacing);
diff --git a/Engine/Plugins/Experimental/PlanarCutPlugin/Source/PlanarCut/Public/FractureAutoUV.h b/Engine/Plugins/Experimental/PlanarCutPlugin/Source/PlanarCut/Public/FractureAutoUV.h
--- a/Engine/Plugins/Experimental/PlanarCutPlugin/Source/PlanarCut/Public/FractureAutoUV.h
+++ b/Engine/Plugins/Experimental/PlanarCutPlugin/Source/PlanarCut/Public/FractureAutoUV.h
@@ -11,6 +11,7 @@ namespace UE { namespace PlanarCut {
 /**
  * Make a UV atlas of non-overlapping UV charts for a geometry collection
  *
+ * @param TargetUVLayer		Which UV layer to update with new UV coordinates
  * @param Collection		The collection to be atlas'd
  * @param UVRes				Target resolution for the atlas
  * @param GutterSize		Space to leave between UV islands, in pixels at the target resolution
@@ -19,6 +20,7 @@ namespace UE { namespace PlanarCut {
  * @param bRecreateUVsForDegenerateIslands If true, detect and fix islands that don't have proper UVs (i.e. UVs all zero or otherwise collapsed to a point)
  */
 bool PLANARCUT_API UVLayout(
+	int32 TargetUVLayer,
 	FGeometryCollection& Collection,
 	int32 UVRes = 1024,
 	float GutterSize = 1,
@@ -66,6 +68,7 @@ struct FTextureAttributeSettings
  * Generate a texture for internal faces based on depth inside surface
  * TODO: add options to texture based on other quantities
  *
+ * @param TargetUVLayer		Which UV layer to take UV coordinates from when creating the new texture
  * @param Collection		The collection to be create a new texture for
  * @param GutterSize		Number of texels to fill outside of UV island borders (values are copied from nearest inside pt)
  * @param BakeAttributes	Which attributes to bake into which color channel
@@ -75,6 +78,7 @@ struct FTextureAttributeSettings
  * @param WhichMaterials	If non-empty, apply texture to the listed material IDs
  */
 void PLANARCUT_API TextureInternalSurfaces(
+	int32 TargetUVLayer,
 	FGeometryCollection& Collection,
 	int32 GutterSize,
 	UE::Geometry::FIndex4i BakeAttributes,
