Fixed CutWithCellMeshes potentially mis-reporting the index of the first geometry it created, instead returning the last of a group of islands if there are many non-touching cells.

diff --git a/Engine/Plugins/Experimental/PlanarCutPlugin/Source/PlanarCut/Private/GeometryMeshConversion.cpp b/Engine/Plugins/Experimental/PlanarCutPlugin/Source/PlanarCut/Private/GeometryMeshConversion.cpp
--- a/Engine/Plugins/Experimental/PlanarCutPlugin/Source/PlanarCut/Private/GeometryMeshConversion.cpp
+++ b/Engine/Plugins/Experimental/PlanarCutPlugin/Source/PlanarCut/Private/GeometryMeshConversion.cpp
@@ -1963,6 +1963,11 @@ int32 FDynamicMeshCollection::CutWithCellMeshes(const FInternalSurfaceMaterials&
 								*BooleanResults[CellIdx] = Island;
 								GeometryToResultMesh.Add(CreatedGeometryIdx, CellIdx);
 							}
+
+							if(FirstIdx == -1)
+							{
+								FirstIdx = CreatedGeometryIdx;
+							}
 						}
 					}
 					else
@@ -1971,10 +1976,11 @@ int32 FDynamicMeshCollection::CutWithCellMeshes(const FInternalSurfaceMaterials&
 						CreatedGeometryIdx = AppendToCollection(Surface.ToCollection, AugBoolResult, CollisionSampleSpacing, Surface.TransformIndex, BoneName, *Collection, InternalMaterialID);
 						CellToGeometry.Add(CellIdx, CreatedGeometryIdx);
 						GeometryToResultMesh.Add(CreatedGeometryIdx, CellIdx);
-					}
-					if (FirstIdx == -1)
-					{
-						FirstIdx = CreatedGeometryIdx;
+
+						if(FirstIdx == -1)
+						{
+							FirstIdx = CreatedGeometryIdx;
+						}
 					}
 				}
 			}
