Add option to target all material IDs (i.e. all geometry) in fracture AutoUV

#rnx

diff --git a/Engine/Plugins/Experimental/ChaosEditor/Source/FractureEditor/Private/FractureToolUV.cpp b/Engine/Plugins/Experimental/ChaosEditor/Source/FractureEditor/Private/FractureToolUV.cpp
--- a/Engine/Plugins/Experimental/ChaosEditor/Source/FractureEditor/Private/FractureToolUV.cpp
+++ b/Engine/Plugins/Experimental/ChaosEditor/Source/FractureEditor/Private/FractureToolUV.cpp
@@ -289,11 +289,15 @@ int32 UFractureToolAutoUV::ExecuteFracture(const FFractureToolContext& FractureC
 
 		int32 UVLayer = AutoUVSettings->GetSelectedChannelIndex();
 
+		UE::PlanarCut::EUseMaterials UseMaterialIDs =
+			AutoUVSettings->TargetMaterialIDs == ETargetMaterialIDs::SelectedIDs ? UE::PlanarCut::EUseMaterials::NoDefaultMaterials :
+			(AutoUVSettings->TargetMaterialIDs == ETargetMaterialIDs::AllIDs ? UE::PlanarCut::EUseMaterials::AllMaterials : UE::PlanarCut::EUseMaterials::OddMaterials);
+
 		if (AutoUVSettings->bDoUVLayout)
 		{
 			UVTask.EnterProgressFrame(1, LOCTEXT("LayOutUVIslands", "Laying out UV islands"));
 			if (!UE::PlanarCut::UVLayout(UVLayer, Collection, OutputRes, AutoUVSettings->GutterSize,
-				AutoUVSettings->TargetMaterialIDs != ETargetMaterialIDs::SelectedIDs,
+				UseMaterialIDs,
 				AutoUVSettings->TargetMaterialIDs == ETargetMaterialIDs::OddIDs ? EmptyMaterialIDs : AutoUVSettings->MaterialIDs))
 			{
 				// failed to do layout
@@ -343,7 +347,7 @@ int32 UFractureToolAutoUV::ExecuteFracture(const FFractureToolContext& FractureC
 		AttribSettings.Curvature_MaxValue = AutoUVSettings->MaxCurvature;
 		AttribSettings.ClearGutterChannel = 3; // default clear the gutters for the alpha channel, so it shows more clearly the island boundaries
 		UE::PlanarCut::TextureInternalSurfaces(UVLayer, Collection, FMath::CeilToInt(AutoUVSettings->GutterSize), Attributes, AttribSettings, ImageBuilder,
-			AutoUVSettings->TargetMaterialIDs != ETargetMaterialIDs::SelectedIDs, 
+			UseMaterialIDs,
 			AutoUVSettings->TargetMaterialIDs == ETargetMaterialIDs::OddIDs ? EmptyMaterialIDs: AutoUVSettings->MaterialIDs);
 
 		UVTask.EnterProgressFrame(1, LOCTEXT("SavingTexture", "Saving result"));
diff --git a/Engine/Plugins/Experimental/ChaosEditor/Source/FractureEditor/Private/FractureToolUV.h b/Engine/Plugins/Experimental/ChaosEditor/Source/FractureEditor/Private/FractureToolUV.h
--- a/Engine/Plugins/Experimental/ChaosEditor/Source/FractureEditor/Private/FractureToolUV.h
+++ b/Engine/Plugins/Experimental/ChaosEditor/Source/FractureEditor/Private/FractureToolUV.h
@@ -38,7 +38,8 @@ enum class ETargetMaterialIDs
 {
 	OddIDs,
 	OddAndSelectedIDs,
-	SelectedIDs
+	SelectedIDs,
+	AllIDs
 };
 
 /** Settings specifically related to the one-time destructive fracturing of a mesh **/
diff --git a/Engine/Plugins/Experimental/PlanarCutPlugin/Source/PlanarCut/Private/FractureAutoUV.cpp b/Engine/Plugins/Experimental/PlanarCutPlugin/Source/PlanarCut/Private/FractureAutoUV.cpp
--- a/Engine/Plugins/Experimental/PlanarCutPlugin/Source/PlanarCut/Private/FractureAutoUV.cpp
+++ b/Engine/Plugins/Experimental/PlanarCutPlugin/Source/PlanarCut/Private/FractureAutoUV.cpp
@@ -109,14 +109,18 @@ namespace UE { namespace PlanarCut {
 
 
 namespace {
-	inline bool IsTriActive(bool bIsVisible, int32 MaterialID, const TSet<int32>& InsideMaterials, bool bActivateInsideTriangles, bool bOddMaterialAreInside)
+	inline bool IsTriActive(bool bIsVisible, int32 MaterialID, const TSet<int32>& InsideMaterials, bool bActivateInsideTriangles, EUseMaterials MaterialsPattern)
 	{
-		if (!bIsVisible)
+		if (!bIsVisible) // never select invisible triangles
 		{
 			return false;
 		}
 
-		if (bOddMaterialAreInside && ((MaterialID % 2) == 1) == bActivateInsideTriangles)
+		if (MaterialsPattern == EUseMaterials::AllMaterials && bActivateInsideTriangles)
+		{
+			return true;
+		}
+		if (MaterialsPattern == EUseMaterials::OddMaterials && ((MaterialID % 2) == 1) == bActivateInsideTriangles)
 		{
 			return true;
 		}
@@ -134,8 +138,7 @@ namespace {
  * @return number of active triangles
  */
 int32 SetActiveTriangles(FGeometryCollection* Collection, TArray<bool>& ActiveTrianglesOut,
-	bool bActivateInsideTriangles,
-	bool bOddMaterialAreInside, TArrayView<int32> WhichMaterialsAreInside)
+	bool bActivateInsideTriangles, EUseMaterials MaterialsPattern, TArrayView<int32> WhichMaterialsAreInside)
 {
 	TSet<int32> InsideMaterials;
 	for (int32 MatID : WhichMaterialsAreInside)
@@ -147,7 +150,7 @@ int32 SetActiveTriangles(FGeometryCollection* Collection, TArray<bool>& ActiveTr
 	int32 NumTriangles = 0;
 	for (int TID = 0; TID < Collection->Indices.Num(); TID++)
 	{
-		bool bIsActive = IsTriActive(Collection->Visible[TID], Collection->MaterialID[TID], InsideMaterials, bActivateInsideTriangles, bOddMaterialAreInside);
+		bool bIsActive = IsTriActive(Collection->Visible[TID], Collection->MaterialID[TID], InsideMaterials, bActivateInsideTriangles, MaterialsPattern);
 		if (bIsActive)
 		{
 			ActiveTrianglesOut[TID] = true;
@@ -374,13 +377,13 @@ bool UVLayout(
 	FGeometryCollection& Collection,
 	int32 UVRes,
 	float GutterSize,
-	bool bOnlyOddMaterials,
+	EUseMaterials MaterialsPattern,
 	TArrayView<int32> WhichMaterials,
 	bool bRecreateUVsForDegenerateIslands
 )
 {
 	TArray<bool> ActiveTriangles;
-	int32 NumActive = SetActiveTriangles(&Collection, ActiveTriangles, true, bOnlyOddMaterials, WhichMaterials);
+	int32 NumActive = SetActiveTriangles(&Collection, ActiveTriangles, true, MaterialsPattern, WhichMaterials);
 	FGeomMesh UVMesh(TargetUVLayer, &Collection, ActiveTriangles, NumActive);
 
 	TArray<TArray<int32>> UVIslands;
@@ -505,12 +508,12 @@ void TextureInternalSurfaces(
 	FIndex4i BakeAttributes,
 	const FTextureAttributeSettings& AttributeSettings,
 	TImageBuilder<FVector4f>& TextureOut,
-	bool bOnlyOddMaterials,
+	EUseMaterials MaterialsPattern,
 	TArrayView<int32> WhichMaterials
 )
 {
 	TArray<bool> ToTextureTriangles;
-	int32 NumTextureTris = SetActiveTriangles(&Collection, ToTextureTriangles, true, bOnlyOddMaterials, WhichMaterials);
+	int32 NumTextureTris = SetActiveTriangles(&Collection, ToTextureTriangles, true, MaterialsPattern, WhichMaterials);
 	FGeomFlatUVMesh UVMesh(TargetUVLayer, &Collection, ToTextureTriangles, NumTextureTris);
 
 	FImageOccupancyMap OccupancyMap;
@@ -522,7 +525,7 @@ void TextureInternalSurfaces(
 
 	// Outside triangles are identified by odd material ID
 	TArray<bool> OutsideTriangles;
-	int32 NumOutsideTris = SetActiveTriangles(&Collection, OutsideTriangles, false, true, {});
+	int32 NumOutsideTris = SetActiveTriangles(&Collection, OutsideTriangles, false, EUseMaterials::OddMaterials, {});
 	FGeomMesh OutsideMesh(ToTextureMesh, OutsideTriangles, NumOutsideTris);
 	TMeshAABBTree3<FGeomMesh> OutsideSpatial(&OutsideMesh);
 
@@ -567,7 +570,7 @@ void TextureInternalSurfaces(
 			{
 				bool bIsTextureTri = IsTriActive(
 					AugmentedDynamicMesh::GetVisibility(Mesh, TID), 
-					Mesh.Attributes()->GetMaterialID()->GetValue(TID), TargetMaterials, true, bOnlyOddMaterials);
+					Mesh.Attributes()->GetMaterialID()->GetValue(TID), TargetMaterials, true, MaterialsPattern);
 				if (!bIsTextureTri)
 				{
 					UV->UnsetTriangle(TID);
diff --git a/Engine/Plugins/Experimental/PlanarCutPlugin/Source/PlanarCut/Public/FractureAutoUV.h b/Engine/Plugins/Experimental/PlanarCutPlugin/Source/PlanarCut/Public/FractureAutoUV.h
--- a/Engine/Plugins/Experimental/PlanarCutPlugin/Source/PlanarCut/Public/FractureAutoUV.h
+++ b/Engine/Plugins/Experimental/PlanarCutPlugin/Source/PlanarCut/Public/FractureAutoUV.h
@@ -8,6 +8,13 @@
 
 namespace UE { namespace PlanarCut {
 
+enum class EUseMaterials
+{
+	AllMaterials,		// Include all materials
+	OddMaterials,		// Include materials w/ odd IDs (+ any manually selected materials)
+	NoDefaultMaterials	// No default materials; Only use manually selected materials
+};
+
 /**
  * Make a UV atlas of non-overlapping UV charts for a geometry collection
  *
@@ -15,7 +22,7 @@ namespace UE { namespace PlanarCut {
  * @param Collection		The collection to be atlas'd
  * @param UVRes				Target resolution for the atlas
  * @param GutterSize		Space to leave between UV islands, in pixels at the target resolution
- * @param bOddMaterials		If true, consider odd-numbered material IDs for UV island layout
+ * @param MaterialsPattern	Which pattern of material IDs to automatically consider for UV island layout
  * @param WhichMaterials	If non-empty, consider listed material IDs for UV island layout
  * @param bRecreateUVsForDegenerateIslands If true, detect and fix islands that don't have proper UVs (i.e. UVs all zero or otherwise collapsed to a point)
  */
@@ -24,7 +31,7 @@ bool PLANARCUT_API UVLayout(
 	FGeometryCollection& Collection,
 	int32 UVRes = 1024,
 	float GutterSize = 1,
-	bool bOddMaterials = true,
+	EUseMaterials MaterialsPattern = EUseMaterials::OddMaterials,
 	TArrayView<int32> WhichMaterials = TArrayView<int32>(),
 	bool bRecreateUVsForDegenerateIslands = true
 );
@@ -75,7 +82,7 @@ struct FTextureAttributeSettings
  * @param BakeAttributes	Which attributes to bake into which color channel
  * @param AttributeSettings	Settings for the BakeAttributes
  * @param TextureOut		Texture to write to
- * @param bOddMaterials		If true, apply texture to odd-numbered material IDs
+ * @param MaterialsPattern	Which pattern of material IDs to apply texture to
  * @param WhichMaterials	If non-empty, apply texture to the listed material IDs
  */
 void PLANARCUT_API TextureInternalSurfaces(
@@ -85,7 +92,7 @@ void PLANARCUT_API TextureInternalSurfaces(
 	UE::Geometry::FIndex4i BakeAttributes,
 	const FTextureAttributeSettings& AttributeSettings,
 	UE::Geometry::TImageBuilder<FVector4f>& TextureOut,
-	bool bOddMaterials = true,
+	EUseMaterials MaterialsPattern = EUseMaterials::OddMaterials,
 	TArrayView<int32> WhichMaterials = TArrayView<int32>()
 );
 
