add v0 of box projection to fracture auto uv tool

#rnx

diff --git a/Engine/Plugins/Experimental/ChaosEditor/Source/FractureEditor/Private/FractureToolUV.cpp b/Engine/Plugins/Experimental/ChaosEditor/Source/FractureEditor/Private/FractureToolUV.cpp
--- a/Engine/Plugins/Experimental/ChaosEditor/Source/FractureEditor/Private/FractureToolUV.cpp
+++ b/Engine/Plugins/Experimental/ChaosEditor/Source/FractureEditor/Private/FractureToolUV.cpp
@@ -59,6 +59,11 @@ void UFractureAutoUVSettings::ChangeNumUVChannels(int32 Delta)
 	}
 }
 
+void UFractureAutoUVSettings::BoxProjectUVs()
+{
+	UFractureToolAutoUV* AutoUVTool = Cast<UFractureToolAutoUV>(OwnerTool.Get());
+	AutoUVTool->BoxProjectUVs();
+}
 
 UFractureToolAutoUV::UFractureToolAutoUV(const FObjectInitializer& ObjInit)
 	: Super(ObjInit)
@@ -153,6 +158,30 @@ void UFractureToolAutoUV::UpdateUVChannels(int32 TargetNumUVChannels)
 	AutoUVSettings->SetNumUVChannels(MinUVChannels);
 }
 
+void UFractureToolAutoUV::BoxProjectUVs()
+{
+	int32 UVLayer = AutoUVSettings->GetSelectedChannelIndex();
+
+	TArray<int32> EmptyMaterialIDs;
+	UE::PlanarCut::EUseMaterials UseMaterialIDs =
+		AutoUVSettings->TargetMaterialIDs == ETargetMaterialIDs::SelectedIDs ? UE::PlanarCut::EUseMaterials::NoDefaultMaterials :
+		(AutoUVSettings->TargetMaterialIDs == ETargetMaterialIDs::AllIDs ? UE::PlanarCut::EUseMaterials::AllMaterials : UE::PlanarCut::EUseMaterials::OddMaterials);
+
+	TSet<UGeometryCollectionComponent*> GeomCompSelection;
+	GetSelectedGeometryCollectionComponents(GeomCompSelection);
+	for (UGeometryCollectionComponent* GeometryCollectionComponent : GeomCompSelection)
+	{
+		UE::PlanarCut::BoxProjectUVs(UVLayer,
+			*GeometryCollectionComponent->GetRestCollection()->GetGeometryCollection(),
+			FVector3d(100, 100, 100),
+			UseMaterialIDs,
+			AutoUVSettings->TargetMaterialIDs == ETargetMaterialIDs::OddIDs ? EmptyMaterialIDs : AutoUVSettings->MaterialIDs);
+
+		GeometryCollectionComponent->MarkRenderDynamicDataDirty();
+		GeometryCollectionComponent->MarkRenderStateDirty();
+	}
+}
+
 void UFractureToolAutoUV::Render(const FSceneView* View, FViewport* Viewport, FPrimitiveDrawInterface* PDI)
 {
 }
diff --git a/Engine/Plugins/Experimental/ChaosEditor/Source/FractureEditor/Private/FractureToolUV.h b/Engine/Plugins/Experimental/ChaosEditor/Source/FractureEditor/Private/FractureToolUV.h
--- a/Engine/Plugins/Experimental/ChaosEditor/Source/FractureEditor/Private/FractureToolUV.h
+++ b/Engine/Plugins/Experimental/ChaosEditor/Source/FractureEditor/Private/FractureToolUV.h
@@ -85,6 +85,9 @@ public:
 
 	void ChangeNumUVChannels(int32 Delta);
 
+	UFUNCTION(CallInEditor, Category = UVIslands, meta = (DisplayName = "Box Project UVs"))
+	void BoxProjectUVs();
+
 	/** Whether to layout UVs (to ensure unique texels for each face), or just use the existing UVs */
 	UPROPERTY(EditAnywhere, Category = Atlas, meta = (DisplayName = "Do UV Layout"))
 	bool bDoUVLayout = true;
@@ -206,6 +209,9 @@ public:
 	// Update number of UV channels for selected geometry collections
 	void UpdateUVChannels(int32 TargetNumUVChannels);
 
+	// Create new UV islands via box projection for the selected geometry collections
+	void BoxProjectUVs();
+
 
 protected:
 	UPROPERTY(EditAnywhere, Category = Slicing)
diff --git a/Engine/Plugins/Experimental/PlanarCutPlugin/Source/PlanarCut/Private/FractureAutoUV.cpp b/Engine/Plugins/Experimental/PlanarCutPlugin/Source/PlanarCut/Private/FractureAutoUV.cpp
--- a/Engine/Plugins/Experimental/PlanarCutPlugin/Source/PlanarCut/Private/FractureAutoUV.cpp
+++ b/Engine/Plugins/Experimental/PlanarCutPlugin/Source/PlanarCut/Private/FractureAutoUV.cpp
@@ -13,6 +13,7 @@
 #include "MeshWeights.h"
 #include "Parameterization/DynamicMeshUVEditor.h"
 #include "Sampling/MeshImageBakingCache.h"
+#include "DynamicMesh/Operations/MergeCoincidentMeshEdges.h"
 
 #include "Implicit/Solidify.h"
 
@@ -369,7 +370,85 @@ struct FGeomFlatUVMesh
 };
 
 
+bool BoxProjectUVs(
+	int32 TargetUVLayer,
+	FGeometryCollection& Collection,
+	const FVector3d& BoxDimensions,
+	EUseMaterials MaterialsPattern,
+	TArrayView<int32> WhichMaterials
+)
+{
+	TArray<int32> TransformIndices;
+	for (int32 GeomIdx = 0; GeomIdx < Collection.TransformIndex.Num(); GeomIdx++)
+	{
+		TransformIndices.Add(Collection.TransformIndex[GeomIdx]);
+	}
+	FDynamicMeshCollection CollectionMeshes(&Collection, TransformIndices, FTransform::Identity, false);
+	TSet<int32> TargetMaterials;
+	for (int32 MatID : WhichMaterials)
+	{
+		TargetMaterials.Add(MatID);
+	}
+
+	int32 NumUVLayers = Collection.NumUVLayers();
+	if (TargetUVLayer >= NumUVLayers)
+	{
+		return false;
+	}
+
+	for (int MeshIdx = 0; MeshIdx < CollectionMeshes.Meshes.Num(); MeshIdx++)
+	{
+		int32 TransformIdx = CollectionMeshes.Meshes[MeshIdx].TransformIndex;
+		FDynamicMesh3& Mesh = CollectionMeshes.Meshes[MeshIdx].AugMesh;
+
+		// TODO: Apply this coincident edge merge once 
+		//FMergeCoincidentMeshEdges EdgeWelder(&Mesh);
+		//EdgeWelder.Apply();
 
+		FMeshNormals::InitializeOverlayToPerVertexNormals(Mesh.Attributes()->PrimaryNormals(), true);
+		AugmentedDynamicMesh::InitializeOverlayToPerVertexUVs(Mesh, NumUVLayers, 0);
+		TArray<int32> TargetTris;
+		for (int TID : Mesh.TriangleIndicesItr())
+		{
+			bool bIsTextureTri = IsTriActive(
+				AugmentedDynamicMesh::GetVisibility(Mesh, TID),
+				Mesh.Attributes()->GetMaterialID()->GetValue(TID), TargetMaterials, true, MaterialsPattern);
+			if (bIsTextureTri)
+			{
+				TargetTris.Add(TID);
+			}
+		}
+		if (TargetTris.Num() == 0)
+		{
+			continue;
+		}
+
+		FDynamicMeshUVEditor UVEd(&Mesh, TargetUVLayer, false);
+		FFrame3d BoxFrame; // defaults to origin / no rotation
+		UVEd.SetTriangleUVsFromBoxProjection(TargetTris, [](const FVector3d& Pos) { return Pos; }, BoxFrame, BoxDimensions);
+
+		Mesh.CompactInPlace();
+
+		// transfer work back to custom UV data of AugmentedDynamicMeshes
+		FDynamicMeshUVOverlay* UVOverlay = Mesh.Attributes()->GetUVLayer(TargetUVLayer);
+
+		for (int ElID : UVOverlay->ElementIndicesItr())
+		{
+			FVector2f UV = UVOverlay->GetElement(ElID);
+			int VID = UVOverlay->GetParentVertex(ElID);
+			checkSlow(VID > -1);
+			if (VID > -1)
+			{
+				AugmentedDynamicMesh::SetUV(Mesh, VID, UV, TargetUVLayer);
+			}
+		}
+		
+	}
+
+	CollectionMeshes.UpdateAllCollections(Collection);
+
+	return true;
+}
 
 
 bool UVLayout(
diff --git a/Engine/Plugins/Experimental/PlanarCutPlugin/Source/PlanarCut/Public/FractureAutoUV.h b/Engine/Plugins/Experimental/PlanarCutPlugin/Source/PlanarCut/Public/FractureAutoUV.h
--- a/Engine/Plugins/Experimental/PlanarCutPlugin/Source/PlanarCut/Public/FractureAutoUV.h
+++ b/Engine/Plugins/Experimental/PlanarCutPlugin/Source/PlanarCut/Public/FractureAutoUV.h
@@ -15,6 +15,23 @@ enum class EUseMaterials
 	NoDefaultMaterials	// No default materials; Only use manually selected materials
 };
 
+/**
+ * Box project UVs
+ *
+ * @param TargetUVLayer		Which UV layer to update
+ * @param Collection		The collection to be box projected
+ * @param BoxDimensions		Scale of projection box
+ * @param MaterialsPattern	Which pattern of material IDs to automatically consider for UV island layout
+ * @param WhichMaterials	If non-empty, consider listed material IDs for UV island layout
+ */
+bool PLANARCUT_API BoxProjectUVs(
+	int32 TargetUVLayer,
+	FGeometryCollection& Collection,
+	const FVector3d& BoxDimensions,
+	EUseMaterials MaterialsPattern = EUseMaterials::OddMaterials,
+	TArrayView<int32> WhichMaterials = TArrayView<int32>()
+);
+
 /**
  * Make a UV atlas of non-overlapping UV charts for a geometry collection
  *
