fix fracture autouv box project output topology + related bugs

#rnx

diff --git a/Engine/Plugins/Experimental/PlanarCutPlugin/Source/PlanarCut/Private/FractureAutoUV.cpp b/Engine/Plugins/Experimental/PlanarCutPlugin/Source/PlanarCut/Private/FractureAutoUV.cpp
--- a/Engine/Plugins/Experimental/PlanarCutPlugin/Source/PlanarCut/Private/FractureAutoUV.cpp
+++ b/Engine/Plugins/Experimental/PlanarCutPlugin/Source/PlanarCut/Private/FractureAutoUV.cpp
@@ -401,12 +401,14 @@ bool BoxProjectUVs(
 		int32 TransformIdx = CollectionMeshes.Meshes[MeshIdx].TransformIndex;
 		FDynamicMesh3& Mesh = CollectionMeshes.Meshes[MeshIdx].AugMesh;
 
-		// TODO: Apply this coincident edge merge once 
-		//FMergeCoincidentMeshEdges EdgeWelder(&Mesh);
-		//EdgeWelder.Apply();
-
 		FMeshNormals::InitializeOverlayToPerVertexNormals(Mesh.Attributes()->PrimaryNormals(), true);
+		AugmentedDynamicMesh::InitializeOverlayToPerVertexTangents(Mesh);
 		AugmentedDynamicMesh::InitializeOverlayToPerVertexUVs(Mesh, NumUVLayers, 0);
+
+		// Apply coincident edge merge so that projection can re-determine the UV islands without having to keep existing seams
+		FMergeCoincidentMeshEdges EdgeWelder(&Mesh);
+		EdgeWelder.Apply();
+
 		TArray<int32> TargetTris;
 		for (int TID : Mesh.TriangleIndicesItr())
 		{
@@ -429,20 +431,8 @@ bool BoxProjectUVs(
 
 		Mesh.CompactInPlace();
 
-		// transfer work back to custom UV data of AugmentedDynamicMeshes
-		FDynamicMeshUVOverlay* UVOverlay = Mesh.Attributes()->GetUVLayer(TargetUVLayer);
-
-		for (int ElID : UVOverlay->ElementIndicesItr())
-		{
-			FVector2f UV = UVOverlay->GetElement(ElID);
-			int VID = UVOverlay->GetParentVertex(ElID);
-			checkSlow(VID > -1);
-			if (VID > -1)
-			{
-				AugmentedDynamicMesh::SetUV(Mesh, VID, UV, TargetUVLayer);
-			}
-		}
-		
+		// Re-split vertices with different UVs/normals/tangents and transfer attributes back (required because we weld edges above)
+		AugmentedDynamicMesh::SplitOverlayAttributesToPerVertex(Mesh, true, true);
 	}
 
 	CollectionMeshes.UpdateAllCollections(Collection);
diff --git a/Engine/Plugins/Experimental/PlanarCutPlugin/Source/PlanarCut/Private/GeometryMeshConversion.cpp b/Engine/Plugins/Experimental/PlanarCutPlugin/Source/PlanarCut/Private/GeometryMeshConversion.cpp
--- a/Engine/Plugins/Experimental/PlanarCutPlugin/Source/PlanarCut/Private/GeometryMeshConversion.cpp
+++ b/Engine/Plugins/Experimental/PlanarCutPlugin/Source/PlanarCut/Private/GeometryMeshConversion.cpp
@@ -100,6 +100,22 @@ namespace AugmentedDynamicMesh
 		}
 	}
 
+	int32 NumEnabledUVChannels(FDynamicMesh3& Mesh)
+	{
+		if (!Mesh.Attributes())
+		{
+			return 0;
+		}
+		for (int32 UVIdx = 0; UVIdx < MAX_NUM_UV_CHANNELS; UVIdx++)
+		{
+			if (!Mesh.Attributes()->HasAttachedAttribute(UVChannelNames[UVIdx]))
+			{
+				return UVIdx;
+			}
+		}
+		return MAX_NUM_UV_CHANNELS;
+	}
+
 	void Augment(FDynamicMesh3& Mesh, int32 NumUVChannels)
 	{
 		Mesh.EnableVertexColors(FVector3f(1, 1, 1));
@@ -235,6 +251,137 @@ namespace AugmentedDynamicMesh
 		Vs->GetValue(VID, V);
 	}
 
+	template<typename FOverlay>
+	void SplitOverlayHelper(UE::Geometry::FDynamicMesh3& Mesh, FOverlay* Overlay)
+	{
+		TArray<int32> ContigTris, ContigGroupLens;
+		TArray<int32> ToSplitTris;
+		TArray<bool> GroupIsLoop;
+
+		auto GetEID = [&Mesh, &Overlay](int32 VID, int32 TID) -> int32
+		{
+			FIndex3i Tri = Mesh.GetTriangle(TID);
+			int32 SubIdx = Tri.IndexOf(VID);
+			check(SubIdx >= 0);
+			return Overlay->GetTriangle(TID)[SubIdx];
+		};
+
+		for (int32 VID = 0, OrigMax = Mesh.MaxVertexID(); VID < OrigMax; VID++)
+		{
+			if (!Mesh.IsVertex(VID))
+			{
+				continue;
+			}
+			Mesh.GetVtxContiguousTriangles(VID, ContigTris, ContigGroupLens, GroupIsLoop);
+			int32 TriStartIdx = 0;
+			for (int32 GroupIdx = 0; GroupIdx < ContigGroupLens.Num(); GroupIdx++)
+			{
+				int32 GroupLen = ContigGroupLens[GroupIdx];
+				bool bIsLoop = GroupIsLoop[GroupIdx];
+				checkSlow(TriStartIdx < ContigTris.Num());
+				int32 TID0 = ContigTris[TriStartIdx];
+				int32 EID0 = GetEID(VID, TID0);
+				int32 LastEID = EID0;
+				bool bPastEID0 = false;
+				ToSplitTris.Reset();
+				for (int32 Idx = TriStartIdx + 1, EndIdx = TriStartIdx + GroupLen; Idx < EndIdx; Idx++)
+				{
+					int32 TID = ContigTris[Idx];
+					int32 EID = GetEID(VID, TID);
+					if (EID != LastEID)
+					{
+						if (ToSplitTris.Num() > 0)
+						{
+							FDynamicMesh3::FVertexSplitInfo SplitInfo;
+							Mesh.SplitVertex(VID, ToSplitTris, SplitInfo);
+							ToSplitTris.Reset();
+						}
+						LastEID = EID;
+						bPastEID0 = true;
+					}
+					if (bPastEID0)
+					{
+						ToSplitTris.Add(TID);
+					}
+				}
+				if (bPastEID0 && ToSplitTris.Num() > 0 && (!bIsLoop || LastEID != EID0))
+				{
+					// one final group at the end
+					FDynamicMesh3::FVertexSplitInfo SplitInfo;
+					Mesh.SplitVertex(VID, ToSplitTris, SplitInfo);
+				}
+				TriStartIdx += GroupLen;
+			}
+		}
+	}
+
+	void SplitOverlayAttributesToPerVertex(UE::Geometry::FDynamicMesh3& Mesh, bool bSplitUVs, bool bSplitNormalsTangents)
+	{
+		if (!ensure(Mesh.HasAttributes()))
+		{
+			return;
+		}
+
+		int32 NumUVLayers = FMath::Min(NumEnabledUVChannels(Mesh), Mesh.Attributes()->NumUVLayers());
+		if (bSplitUVs)
+		{
+			for (int32 LayerIdx = 0; LayerIdx < NumUVLayers; LayerIdx++)
+			{
+				FDynamicMeshUVOverlay* UVOverlay = Mesh.Attributes()->GetUVLayer(LayerIdx);
+				SplitOverlayHelper<FDynamicMeshUVOverlay>(Mesh, UVOverlay);
+			}
+		}
+
+		if (bSplitNormalsTangents)
+		{
+			for (int32 LayerIdx = 0; LayerIdx < Mesh.Attributes()->NumNormalLayers(); LayerIdx++)
+			{
+				SplitOverlayHelper<FDynamicMeshNormalOverlay>(Mesh, Mesh.Attributes()->GetNormalLayer(0));
+			}
+		}
+
+		FDynamicMeshEditor Editor(&Mesh);
+		FDynamicMeshEditResult EditResult;
+		Editor.SplitBowties(EditResult);
+
+		// copy back attributes
+		for (int32 LayerIdx = 0; LayerIdx < NumUVLayers; LayerIdx++)
+		{
+			FDynamicMeshUVOverlay* UVOverlay = Mesh.Attributes()->GetUVLayer(LayerIdx);
+			TDynamicMeshVertexAttribute<float, 2>* UVs =
+				static_cast<TDynamicMeshVertexAttribute<float, 2>*>(Mesh.Attributes()->GetAttachedAttribute(UVChannelNames[LayerIdx]));
+			for (int32 EID : UVOverlay->ElementIndicesItr())
+			{
+				UVs->SetValue(UVOverlay->GetParentVertex(EID), UVOverlay->GetElement(EID));
+			}
+		}
+		if (Mesh.Attributes()->NumNormalLayers() > 0)
+		{
+			FDynamicMeshNormalOverlay* Overlay = Mesh.Attributes()->GetNormalLayer(0);
+			for (int32 EID : Overlay->ElementIndicesItr())
+			{
+				Mesh.SetVertexNormal(Overlay->GetParentVertex(EID), Overlay->GetElement(EID));
+			}
+		}
+		if (Mesh.Attributes()->HasTangentSpace())
+		{
+			TDynamicMeshVertexAttribute<float, 3>* Us =
+				static_cast<TDynamicMeshVertexAttribute<float, 3>*>(Mesh.Attributes()->GetAttachedAttribute(TangentUAttribName));
+			TDynamicMeshVertexAttribute<float, 3>* Vs =
+				static_cast<TDynamicMeshVertexAttribute<float, 3>*>(Mesh.Attributes()->GetAttachedAttribute(TangentVAttribName));
+			TDynamicMeshVertexAttribute<float, 3>* Tangents[] = { Us, Vs };
+			for (int Idx = 1; Idx < 3; Idx++)
+			{
+				FDynamicMeshNormalOverlay* Overlay = Mesh.Attributes()->GetNormalLayer(Idx);
+				TDynamicMeshVertexAttribute<float, 3>* Tangent = Tangents[Idx - 1];
+				for (int32 EID : Overlay->ElementIndicesItr())
+				{
+					Tangent->SetValue(Overlay->GetParentVertex(EID), Overlay->GetElement(EID));
+				}
+			}
+		}
+	}
+
 	void InitializeOverlayToPerVertexUVs(FDynamicMesh3& Mesh, int32 NumUVLayers, int32 FirstUVLayer)
 	{
 		Mesh.Attributes()->SetNumUVLayers(NumUVLayers);
@@ -2214,7 +2361,8 @@ bool FDynamicMeshCollection::UpdateCollection(const FTransform& ToCollection, FD
 	int32 OldVertexCount = Output.VertexCount[GeometryIdx];
 	int32 OldTriangleCount = Output.FaceCount[GeometryIdx];
 
-	int32 UVLayerCount = Output.NumUVLayers();
+	int32 UVLayerCount = AugmentedDynamicMesh::NumEnabledUVChannels(Mesh);
+	Output.SetNumUVLayers(UVLayerCount);
 
 	int32 NewVertexCount = Mesh.VertexCount();
 	int32 NewTriangleCount = Mesh.TriangleCount();
diff --git a/Engine/Plugins/Experimental/PlanarCutPlugin/Source/PlanarCut/Private/GeometryMeshConversion.h b/Engine/Plugins/Experimental/PlanarCutPlugin/Source/PlanarCut/Private/GeometryMeshConversion.h
--- a/Engine/Plugins/Experimental/PlanarCutPlugin/Source/PlanarCut/Private/GeometryMeshConversion.h
+++ b/Engine/Plugins/Experimental/PlanarCutPlugin/Source/PlanarCut/Private/GeometryMeshConversion.h
@@ -39,6 +39,7 @@ namespace AugmentedDynamicMesh
 	void PLANARCUT_API InitializeOverlayToPerVertexTangents(UE::Geometry::FDynamicMesh3& Mesh);
 	void PLANARCUT_API ComputeTangents(UE::Geometry::FDynamicMesh3& Mesh, bool bOnlyOddMaterials, const TArrayView<const int32>& WhichMaterials, bool bRecomputeNormals = true);
 	void PLANARCUT_API AddCollisionSamplesPerComponent(UE::Geometry::FDynamicMesh3& Mesh, double Spacing);
+	void PLANARCUT_API SplitOverlayAttributesToPerVertex(UE::Geometry::FDynamicMesh3& Mesh, bool bSplitUVs = true, bool bSplitTangents = true);
 }
 
 
diff --git a/Engine/Source/Runtime/Experimental/Chaos/Private/GeometryCollection/GeometryCollectionAlgo.cpp b/Engine/Source/Runtime/Experimental/Chaos/Private/GeometryCollection/GeometryCollectionAlgo.cpp
--- a/Engine/Source/Runtime/Experimental/Chaos/Private/GeometryCollection/GeometryCollectionAlgo.cpp
+++ b/Engine/Source/Runtime/Experimental/Chaos/Private/GeometryCollection/GeometryCollectionAlgo.cpp
@@ -839,13 +839,43 @@ namespace GeometryCollectionAlgo
 		// fix face/vertex counts and vertex->transform (bone) map
 		for (int32 GeometryIdx = 0; GeometryIdx < NumGeometries; GeometryIdx++)
 		{
+			int32 TransformIdx = GeometryCollection->TransformIndex[GeometryIdx];
 			GeometryCollection->FaceCount[GeometryIdx] = FaceCounts[GeometryIdx];
 			GeometryCollection->VertexCount[GeometryIdx] = VertexCounts[GeometryIdx];
 			int32 VertexStart = GeometryCollection->VertexStart[GeometryIdx];
 			int32 VertexEnd = VertexStart + GeometryCollection->VertexCount[GeometryIdx];
 			for (int32 VertexIdx = VertexStart; VertexIdx < VertexEnd; VertexIdx++)
 			{
-				GeometryCollection->BoneMap[VertexIdx] = GeometryCollection->TransformIndex[GeometryIdx];
+				GeometryCollection->BoneMap[VertexIdx] = TransformIdx;
+			}
+		}
+		for (int32 GeometryIdx = 0; GeometryIdx < NumGeometries; GeometryIdx++)
+		{
+			int32 TransformIdx = GeometryCollection->TransformIndex[GeometryIdx];
+			// the vertex remapping can leave faces that still refer to 'deleted' vertices
+			// these faces are then pointing to vertices that aren't in the same geometry
+			// the intent is that all resized geometry will be re-written by the caller, afterwards
+			// but leaving these broken faces is dangerous and prevents validation in the meantime
+			// so we 'fix' them by writing a degenerate (but w/in geo) face in these cases
+			int32 FaceStart = GeometryCollection->FaceStart[GeometryIdx];
+			int32 FaceEnd = FaceStart + GeometryCollection->FaceCount[GeometryIdx];
+			int32 VertexStart = GeometryCollection->VertexStart[GeometryIdx];
+			FIntVector ReplacementFace(VertexStart, VertexStart, VertexStart);
+			for (int32 FaceIdx = FaceStart; FaceIdx < FaceEnd; FaceIdx++)
+			{
+				FIntVector& Face = GeometryCollection->Indices[FaceIdx];
+				bool bFaceUsesDeletedVertex = false;
+				for (int SubIdx = 0; SubIdx < 3; SubIdx++)
+				{
+					int32 VertIdx = Face[SubIdx];
+					bFaceUsesDeletedVertex |=
+						(GeometryCollection->BoneMap[VertIdx] != TransformIdx) |
+						(VertIdx >= NewNumVertices);
+				}
+				if (bFaceUsesDeletedVertex)
+				{
+					Face = ReplacementFace;
+				}
 			}
 		}
 
