add a "fix tiny geo" tool, to merge tiny geometry from fracture into larger neighboring geometry

diff --git a/Engine/Plugins/Experimental/ChaosEditor/Source/FractureEditor/Private/FractureEditorCommands.h b/Engine/Plugins/Experimental/ChaosEditor/Source/FractureEditor/Private/FractureEditorCommands.h
--- a/Engine/Plugins/Experimental/ChaosEditor/Source/FractureEditor/Private/FractureEditorCommands.h
+++ b/Engine/Plugins/Experimental/ChaosEditor/Source/FractureEditor/Private/FractureEditorCommands.h
@@ -66,6 +66,7 @@ class FFractureEditorCommands : public TCommands<FFractureEditorCommands>
 		TSharedPtr< FUICommandInfo > ConvertToMesh;
 		TSharedPtr< FUICommandInfo > Validate;
 		TSharedPtr< FUICommandInfo > MakeConvex;
+		TSharedPtr< FUICommandInfo > FixTinyGeo;
 
 		// Property Commands
 		TSharedPtr< FUICommandInfo > SetInitialDynamicState;
diff --git a/Engine/Plugins/Experimental/ChaosEditor/Source/FractureEditor/Private/FractureEditorModeToolkit.cpp b/Engine/Plugins/Experimental/ChaosEditor/Source/FractureEditor/Private/FractureEditorModeToolkit.cpp
--- a/Engine/Plugins/Experimental/ChaosEditor/Source/FractureEditor/Private/FractureEditorModeToolkit.cpp
+++ b/Engine/Plugins/Experimental/ChaosEditor/Source/FractureEditor/Private/FractureEditorModeToolkit.cpp
@@ -535,6 +535,7 @@ void FFractureEditorModeToolkit::BuildToolPalette(FName PaletteIndex, class FToo
 		ToolbarBuilder.AddToolBarButton(Commands.ConvertToMesh);
 		ToolbarBuilder.AddToolBarButton(Commands.Validate);
 		ToolbarBuilder.AddToolBarButton(Commands.MakeConvex);
+		ToolbarBuilder.AddToolBarButton(Commands.FixTinyGeo);
 	}
 	else if (PaletteIndex == TEXT("Properties"))
 	{
diff --git a/Engine/Plugins/Experimental/ChaosEditor/Source/FractureEditor/Private/FractureToolFixTinyGeo.cpp b/Engine/Plugins/Experimental/ChaosEditor/Source/FractureEditor/Private/FractureToolFixTinyGeo.cpp
new file mode 100644
--- /dev/null
+++ b/Engine/Plugins/Experimental/ChaosEditor/Source/FractureEditor/Private/FractureToolFixTinyGeo.cpp
@@ -0,0 +1,202 @@
+// Copyright Epic Games, Inc. All Rights Reserved.
+
+#include "FractureToolFixTinyGeo.h"
+#include "FractureEditorStyle.h"
+#include "PlanarCut.h"
+#include "GeometryCollection/GeometryCollection.h"
+#include "GeometryCollection/GeometryCollectionObject.h"
+#include "GeometryCollection/GeometryCollectionAlgo.h"
+#include "FractureEditorModeToolkit.h"
+#include "FractureToolContext.h"
+
+#define LOCTEXT_NAMESPACE "FractureFixTinyGeo"
+
+
+UFractureToolFixTinyGeo::UFractureToolFixTinyGeo(const FObjectInitializer& ObjInit) 
+	: Super(ObjInit) 
+{
+	TinyGeoSettings = NewObject<UFractureTinyGeoSettings>(GetTransientPackage(), UFractureTinyGeoSettings::StaticClass());
+	TinyGeoSettings->OwnerTool = this;
+}
+
+FText UFractureToolFixTinyGeo::GetDisplayText() const
+{
+	return FText(NSLOCTEXT("FixTinyGeo", "FractureToolFixTinyGeo", "Fix Too-Small Geometry")); 
+}
+
+FText UFractureToolFixTinyGeo::GetTooltipText() const 
+{
+	return FText(NSLOCTEXT("FixTinyGeo", "FractureToolFixTinyGeoTooltip", "The TinyGeo tool glues \"too small\" pieces of geometry onto their neighbors."));
+}
+
+FSlateIcon UFractureToolFixTinyGeo::GetToolIcon() const 
+{
+	return FSlateIcon("FractureEditorStyle", "FractureEditor.FixTinyGeo");
+}
+
+void UFractureToolFixTinyGeo::RegisterUICommand( FFractureEditorCommands* BindingContext ) 
+{
+	UI_COMMAND_EXT( BindingContext, UICommandInfo, "FixTinyGeo", "TinyGeo", "FixTinyGeo", EUserInterfaceActionType::ToggleButton, FInputChord() );
+	BindingContext->FixTinyGeo = UICommandInfo;
+}
+
+void UFractureToolFixTinyGeo::Render(const FSceneView* View, FViewport* Viewport, FPrimitiveDrawInterface* PDI)
+{
+	for (FBox& Box : ToRemoveBounds)
+	{
+		FVector B000 = Box.Min;
+		FVector B111 = Box.Max;
+		FVector B011(B000.X, B111.Y, B111.Z);
+		FVector B101(B111.X, B000.Y, B111.Z);
+		FVector B110(B111.X, B111.Y, B000.Z);
+		FVector B001(B000.X, B000.Y, B111.Z);
+		FVector B010(B000.X, B111.Y, B000.Z);
+		FVector B100(B111.X, B000.Y, B000.Z);
+		PDI->DrawLine(B000, B100, FLinearColor::Red, SDPG_Foreground, 0.0f, 0.001f);
+		PDI->DrawLine(B000, B010, FLinearColor::Red, SDPG_Foreground, 0.0f, 0.001f);
+		PDI->DrawLine(B000, B001, FLinearColor::Red, SDPG_Foreground, 0.0f, 0.001f);
+		PDI->DrawLine(B111, B011, FLinearColor::Red, SDPG_Foreground, 0.0f, 0.001f);
+		PDI->DrawLine(B111, B101, FLinearColor::Red, SDPG_Foreground, 0.0f, 0.001f);
+		PDI->DrawLine(B111, B110, FLinearColor::Red, SDPG_Foreground, 0.0f, 0.001f);
+		PDI->DrawLine(B001, B101, FLinearColor::Red, SDPG_Foreground, 0.0f, 0.001f);
+		PDI->DrawLine(B001, B011, FLinearColor::Red, SDPG_Foreground, 0.0f, 0.001f);
+		PDI->DrawLine(B110, B100, FLinearColor::Red, SDPG_Foreground, 0.0f, 0.001f);
+		PDI->DrawLine(B110, B010, FLinearColor::Red, SDPG_Foreground, 0.0f, 0.001f);
+		PDI->DrawLine(B100, B101, FLinearColor::Red, SDPG_Foreground, 0.0f, 0.001f);
+		PDI->DrawLine(B010, B011, FLinearColor::Red, SDPG_Foreground, 0.0f, 0.001f);
+	}
+	
+}
+
+TArray<UObject*> UFractureToolFixTinyGeo::GetSettingsObjects() const
+ {
+	TArray<UObject*> Settings;
+	Settings.Add(TinyGeoSettings);
+	Settings.Add(CollisionSettings);
+	return Settings;
+}
+
+
+void UFractureToolFixTinyGeo::FractureContextChanged()
+{
+	UpdateDefaultRandomSeed();
+	TArray<FFractureToolContext> FractureContexts = GetFractureToolContexts();
+
+	ToRemoveBounds.Reset();
+
+	for (const FFractureToolContext& FractureContext : FractureContexts)
+	{
+		FGeometryCollection& Collection = *FractureContext.GetGeometryCollection();
+
+		TArray<double> Volumes;
+		FindBoneVolumes(
+			Collection,
+			TArrayView<int32>(), /*Empty array => use all transforms*/
+			Volumes,
+			VolDimScale
+		);
+
+		double MinVolume = GetMinVolume(Volumes);
+
+		TArray<int32> SmallIndices;
+		FindSmallBones(
+			Collection,
+			TArrayView<int32>(),
+			Volumes,
+			MinVolume,
+			SmallIndices
+		);
+
+		FTransform OuterTransform = FractureContext.GetTransform();
+		for (int32 TransformIdx : SmallIndices) // small transforms
+		{
+			FTransform InnerTransform = GeometryCollectionAlgo::GlobalMatrix(Collection.Transform, Collection.Parent, TransformIdx);
+			if (Collection.HasAttribute("ExplodedVector", FGeometryCollection::TransformGroup))
+			{
+				TManagedArray<FVector3f>& ExplodedVectors = Collection.GetAttribute<FVector3f>("ExplodedVector", FGeometryCollection::TransformGroup);
+				InnerTransform = InnerTransform * FTransform(ExplodedVectors[TransformIdx]);
+			}
+
+			FTransform CombinedTransform = InnerTransform * OuterTransform;
+
+			int32 GeometryIdx = Collection.TransformToGeometryIndex[TransformIdx];
+			if (GeometryIdx == INDEX_NONE)
+			{
+				ensure(false); // already filtered for small geo, shouldn't have any missing-geo transforms in the list
+				continue;
+			}
+			int32 VStart = Collection.VertexStart[GeometryIdx];
+			int32 VEnd = VStart + Collection.VertexCount[GeometryIdx];
+			FBox Bounds(EForceInit::ForceInit);
+			for (int32 VIdx = VStart; VIdx < VEnd; VIdx++)
+			{
+				Bounds += CombinedTransform.TransformPosition(Collection.Vertex[VIdx]);
+			}
+			ToRemoveBounds.Add(Bounds);
+		}
+	}
+}
+
+double UFractureToolFixTinyGeo::GetMinVolume(TArray<double>& Volumes)
+{
+	double MinVolume = 0;
+	if (TinyGeoSettings->SelectionMethod == EGeometrySelectionMethod::VolumeCubeRoot)
+	{
+		MinVolume = TinyGeoSettings->MinVolumeCubeRoot * VolDimScale;
+		MinVolume = MinVolume * MinVolume * MinVolume;
+	}
+	else // EGeometrySelectionMethod::RelativeVolume
+	{
+		double VolumeSum = 0;
+		for (double Volume : Volumes)
+		{
+			VolumeSum += Volume;
+		}
+		MinVolume = FMath::Pow(VolumeSum, 1.0 / 3.0) * TinyGeoSettings->RelativeVolume;
+		MinVolume = MinVolume * MinVolume * MinVolume;
+	}
+	return MinVolume;
+}
+
+int32 UFractureToolFixTinyGeo::ExecuteFracture(const FFractureToolContext& FractureContext)
+{
+	int32 UpdateIdx = INDEX_NONE;
+	if (FractureContext.GetGeometryCollection().IsValid())
+	{
+		TArray<int32> TransformIndices; // left empty to refer to all indices
+
+		TArray<double> Volumes;
+		FGeometryCollection& Collection = *FractureContext.GetGeometryCollection();
+		FindBoneVolumes(
+			Collection,
+			TransformIndices,
+			Volumes,
+			VolDimScale
+		);
+
+		double MinVolume = GetMinVolume(Volumes);
+
+		TArray<int32> SmallIndices;
+		FindSmallBones(
+			Collection,
+			TransformIndices,
+			Volumes,
+			MinVolume,
+			SmallIndices
+		);
+		MergeBones(
+			Collection,
+			TransformIndices,
+			Volumes,
+			MinVolume,
+			SmallIndices,
+			false /*bUnionJoinedPieces*/ // Note: Union-ing the pieces is nicer in theory, but can leave cracks and non-manifold garbage
+		);
+		UpdateIdx = 0;
+	}
+
+	return UpdateIdx;
+}
+
+
+#undef LOCTEXT_NAMESPACE
\ No newline at end of file
diff --git a/Engine/Plugins/Experimental/ChaosEditor/Source/FractureEditor/Private/FractureToolFixTinyGeo.h b/Engine/Plugins/Experimental/ChaosEditor/Source/FractureEditor/Private/FractureToolFixTinyGeo.h
new file mode 100644
--- /dev/null
+++ b/Engine/Plugins/Experimental/ChaosEditor/Source/FractureEditor/Private/FractureToolFixTinyGeo.h
@@ -0,0 +1,85 @@
+// Copyright Epic Games, Inc. All Rights Reserved.
+
+#pragma once
+
+#include "FractureToolCutter.h"
+#include "Engine/StaticMeshActor.h"
+
+#include "FractureToolFixTinyGeo.generated.h"
+
+class FFractureToolContext;
+
+UENUM()
+enum class EGeometrySelectionMethod
+{
+	VolumeCubeRoot,
+	RelativeVolume
+	// TODO: also have a manual selection method?
+};
+
+/** Settings controlling how convex hulls are generated for geometry collections */
+UCLASS(config = EditorPerProjectUserSettings)
+class UFractureTinyGeoSettings : public UFractureToolSettings
+{
+public:
+
+	GENERATED_BODY()
+
+	UFractureTinyGeoSettings(const FObjectInitializer& ObjInit)
+		: Super(ObjInit)
+	{}
+
+	UPROPERTY(EditAnywhere, Category = FilterSettings)
+	EGeometrySelectionMethod SelectionMethod = EGeometrySelectionMethod::RelativeVolume;
+
+	/** If volume is less than this value cubed, geometry should be merged into neighbors -- i.e. a value of 2 merges geometry smaller than a 2x2x2 cube */
+	UPROPERTY(EditAnywhere, Category = FilterSettings, meta = (ClampMin = ".00001", UIMin = ".1", UIMax = "10", EditCondition = "SelectionMethod == EGeometrySelectionMethod::VolumeCubeRoot", EditConditionHides))
+	double MinVolumeCubeRoot = 1;
+
+	/** If (cube root of) volume relative to the overall geometry (cube root) volume is less than this, geometry should be merged into neighbors */
+	UPROPERTY(EditAnywhere, Category = FilterSettings, meta = (ClampMin = ".001", UIMax = ".1", ClampMax = "1.0", EditCondition = "SelectionMethod == EGeometrySelectionMethod::RelativeVolume", EditConditionHides))
+	double RelativeVolume = .01;
+
+};
+
+
+
+// Note this tool doesn't actually fracture, but it does remake pieces of geometry and shares a lot of machinery with the fracture tools
+UCLASS(DisplayName = "Fix Tiny Geometry Tool", Category = "FractureTools")
+class UFractureToolFixTinyGeo : public UFractureToolCutterBase
+{
+public:
+	GENERATED_BODY()
+
+	UFractureToolFixTinyGeo(const FObjectInitializer& ObjInit);
+
+	// UFractureTool Interface
+	virtual FText GetDisplayText() const override;
+	virtual FText GetTooltipText() const override;
+	virtual FSlateIcon GetToolIcon() const override;
+
+	virtual FText GetApplyText() const override { return FText(NSLOCTEXT("FixTinyGeo", "ExecuteFixTinyGeo", "Merge Geometry")); }
+
+	void Render(const FSceneView* View, FViewport* Viewport, FPrimitiveDrawInterface* PDI) override;
+
+	virtual void RegisterUICommand(FFractureEditorCommands* BindingContext) override;
+
+	virtual TArray<UObject*> GetSettingsObjects() const override;
+
+	virtual void FractureContextChanged() override;
+	virtual int32 ExecuteFracture(const FFractureToolContext& FractureContext) override;
+
+private:
+
+	UPROPERTY(EditAnywhere, Category = FixGeo)
+	TObjectPtr<UFractureTinyGeoSettings> TinyGeoSettings;
+
+
+	TArray<FBox> ToRemoveBounds;
+
+	double GetMinVolume(TArray<double>& Volumes);
+	const double VolDimScale = .01; // compute volumes in meters instead of cm, for saner units at typical scales
+
+};
+
+
diff --git a/Engine/Plugins/Experimental/PlanarCutPlugin/Source/PlanarCut/Private/GeometryMeshConversion.cpp b/Engine/Plugins/Experimental/PlanarCutPlugin/Source/PlanarCut/Private/GeometryMeshConversion.cpp
--- a/Engine/Plugins/Experimental/PlanarCutPlugin/Source/PlanarCut/Private/GeometryMeshConversion.cpp
+++ b/Engine/Plugins/Experimental/PlanarCutPlugin/Source/PlanarCut/Private/GeometryMeshConversion.cpp
@@ -1571,9 +1571,10 @@ void FDynamicMeshCollection::Init(const FGeometryCollection* Collection, const T
 
 	for (int32 TransformIdx : TransformIndices)
 	{
-		if (Collection->Children[TransformIdx].Num() > 0)
+		int32 GeometryIdx = Collection->TransformToGeometryIndex[TransformIdx];
+		if (GeometryIdx == INDEX_NONE)
 		{
-			// only store the meshes of leaf nodes
+			// only store the mesh if there is associated geometry
 			continue;
 		}
 
@@ -1584,8 +1585,7 @@ void FDynamicMeshCollection::Init(const FGeometryCollection* Collection, const T
 		MeshData.TransformIndex = TransformIdx;
 		MeshData.ToCollection = FTransform(CollectionToLocal.Inverse());
 		FDynamicMesh3& Mesh = MeshData.AugMesh;
-
-		int32 GeometryIdx = Collection->TransformToGeometryIndex[TransformIdx];
+		
 		Mesh.EnableAttributes();
 		Mesh.Attributes()->EnableMaterialID();
 
diff --git a/Engine/Plugins/Experimental/PlanarCutPlugin/Source/PlanarCut/Private/PlanarCut.cpp b/Engine/Plugins/Experimental/PlanarCutPlugin/Source/PlanarCut/Private/PlanarCut.cpp
--- a/Engine/Plugins/Experimental/PlanarCutPlugin/Source/PlanarCut/Private/PlanarCut.cpp
+++ b/Engine/Plugins/Experimental/PlanarCutPlugin/Source/PlanarCut/Private/PlanarCut.cpp
@@ -15,6 +15,7 @@
 
 #include "GeometryCollection/GeometryCollectionAlgo.h"
 #include "GeometryCollection/GeometryCollectionClusteringUtility.h"
+#include "GeometryCollection/GeometryCollectionProximityUtility.h"
 
 #include "DisjointSet.h"
 
@@ -760,6 +761,450 @@ int32 CutWithMesh(
 }
 
 
+void FindBoneVolumes(
+	FGeometryCollection& Collection,
+	const TArrayView<const int32>& TransformIndices,
+	TArray<double>& OutVolumes,
+	double ScalePerDimension
+)
+{
+	OutVolumes.Reset();
+
+	TArray<FTransform> Transforms;
+	TArray<int32> TransformIndicesArray(TransformIndices);
+	if (TransformIndicesArray.Num() == 0)
+	{
+		for (int32 TransformIdx = 0; TransformIdx < Collection.TransformToGeometryIndex.Num(); TransformIdx++)
+		{
+			TransformIndicesArray.Add(TransformIdx);
+		}
+	}
+	GeometryCollectionAlgo::GlobalMatrices(Collection.Transform, Collection.Parent, TransformIndicesArray, Transforms);
+
+	auto GetVolume = [](const FGeometryCollection& Collection, int32 GeomIdx, const FTransform& Transform,
+		double DimScaleFactor = 1) -> double
+	{
+		int32 VStart = Collection.VertexStart[GeomIdx];
+		int32 VEnd = VStart + Collection.VertexCount[GeomIdx];
+		if (VStart == VEnd)
+		{
+			return 0.0;
+		}
+		FVector3d Center = FVector::ZeroVector;
+		for (int32 VIdx = VStart; VIdx < VEnd; VIdx++)
+		{
+			FVector Pos = Transform.TransformPosition(Collection.Vertex[VIdx]);
+			Center += (FVector3d)Pos;
+		}
+		Center /= double(VEnd - VStart);
+		int32 FStart = Collection.FaceStart[GeomIdx];
+		int32 FEnd = FStart + Collection.FaceCount[GeomIdx];
+		double VolOut = 0;
+		for (int32 FIdx = FStart; FIdx < FEnd; FIdx++)
+		{
+			FIntVector Tri = Collection.Indices[FIdx];
+			FVector3d V0 = (FVector3d)Transform.TransformPosition(Collection.Vertex[Tri.X]);
+			FVector3d V1 = (FVector3d)Transform.TransformPosition(Collection.Vertex[Tri.Y]);
+			FVector3d V2 = (FVector3d)Transform.TransformPosition(Collection.Vertex[Tri.Z]);
+
+			// add volume of the tetrahedron formed by the triangles and the reference point
+			FVector3d V1mRef = (V1 - Center) * DimScaleFactor;
+			FVector3d V2mRef = (V2 - Center) * DimScaleFactor;
+			FVector3d N = V2mRef.Cross(V1mRef);
+
+			VolOut += ((V0 - Center) * DimScaleFactor).Dot(N) / 6.0;
+		}
+		return VolOut;
+	};
+
+	OutVolumes.SetNum(TransformIndicesArray.Num());
+	ParallelFor(TransformIndicesArray.Num(), [&](int32 Idx)
+		{
+			int32 TransformIdx = TransformIndicesArray[Idx];
+			int32 GeomIdx = Collection.TransformToGeometryIndex[TransformIdx];
+			if (GeomIdx == -1)
+			{
+				OutVolumes[Idx] = 0.0;
+			}
+			else
+			{
+				OutVolumes[Idx] = GetVolume(Collection, GeomIdx, Transforms[Idx], ScalePerDimension);
+			}
+		});
+}
+
+
+void FindSmallBones(
+	FGeometryCollection& Collection,
+	const TArrayView<const int32>& TransformIndices,
+	const TArrayView<const double>& Volumes,
+	double MinVolume,
+	TArray<int32>& OutSmallBones
+)
+{
+	OutSmallBones.Reset();
+
+	auto AddIdx = [&Collection, &Volumes, &MinVolume, &OutSmallBones](int32 TransformIdx)
+	{
+		if (Collection.TransformToGeometryIndex[TransformIdx] > -1 && Volumes[TransformIdx] < MinVolume)
+		{
+			OutSmallBones.Add(TransformIdx);
+		}
+	};
+
+	TArray<FTransform> Transforms;
+	if (TransformIndices.Num() == 0)
+	{
+		int32 NumTransforms = Collection.TransformToGeometryIndex.Num();
+		if (!ensure(Volumes.Num() == NumTransforms))
+		{
+			return;
+		}
+		for (int32 TransformIdx = 0; TransformIdx < NumTransforms; TransformIdx++)
+		{
+			AddIdx(TransformIdx);
+		}
+	}
+	else
+	{
+		if (!ensure(Volumes.Num() == TransformIndices.Num()))
+		{
+			return;
+		}
+		for (int32 TransformIdx : TransformIndices)
+		{
+			AddIdx(TransformIdx);
+		}
+	}
+}
+
+
+int32 MergeBones(
+	FGeometryCollection& Collection,
+	const TArrayView<const int32>& TransformIndicesView,
+	const TArrayView<const double>& Volumes,
+	double MinVolume,
+	const TArrayView<const int32>& SmallTransformIndices,
+	bool bUnionJoinedPieces
+)
+{
+	FTransform CellsToWorld = FTransform::Identity;
+
+	FGeometryCollectionProximityUtility::UpdateProximity(&Collection);
+	TManagedArray<TSet<int32>>& Proximity = Collection.GetAttribute<TSet<int32>>("Proximity", FGeometryCollection::GeometryGroup);
+
+	// local array so we can populate it with all transforms if input was empty
+	TArray<int32> TransformIndices(TransformIndicesView);
+	if (TransformIndices.Num() == 0)
+	{
+		for (int32 TransformIdx = 0; TransformIdx < Collection.TransformToGeometryIndex.Num(); TransformIdx++)
+		{
+			TransformIndices.Add(TransformIdx);
+		}
+	}
+	if (!ensure(TransformIndices.Num() == Volumes.Num()))
+	{
+		return INDEX_NONE;
+	}
+
+	struct FRemoveGroup
+	{
+		int32 MergeTo = INDEX_NONE;
+		double MergeTargetVolume = 0;
+		double TotalVolume = 0;
+		bool bRemoveMergeTarget = false;
+		TArray<int32> ToRemove;
+
+		FRemoveGroup() {}
+		FRemoveGroup(const TMap<int32, double>& GeomVolMaps, double MinVolume, int32 SmallIdx, int32 BigIdx)
+		{
+			ToRemove.Add(SmallIdx);
+			double BigVol = GeomVolMaps[BigIdx];
+			if (BigVol < MinVolume)
+			{
+				ToRemove.Add(BigIdx);
+				bRemoveMergeTarget = true;
+			}
+			MergeTargetVolume = BigVol;
+			TotalVolume = BigVol + GeomVolMaps[SmallIdx];
+			MergeTo = BigIdx;
+		}
+
+		bool IsValid()
+		{
+			return MergeTo != INDEX_NONE;
+		}
+
+		void UpdateMergeTarget(int Idx, double Volume)
+		{
+			if (bRemoveMergeTarget)
+			{
+				if (Volume > MergeTargetVolume)
+				{
+					MergeTo = Idx;
+					MergeTargetVolume = Volume;
+				}
+			}
+		}
+
+		// add a too-small geometry to this group
+		void AddSmall(const TMap<int32, double>& GeomVolMaps, int32 SmallIdx)
+		{
+			double SmallVol = GeomVolMaps[SmallIdx];
+			TotalVolume += SmallVol;
+			UpdateMergeTarget(SmallIdx, SmallVol);
+			checkSlow(!ToRemove.Contains(SmallIdx));
+			ToRemove.Add(SmallIdx);
+		}
+
+		// add a neighbor of a too-small geometry to an existing group
+		void AddBig(const TMap<int32, double>& GeomVolMaps, double MinVolume, int32 BigIdx)
+		{
+			double BigVol = GeomVolMaps[BigIdx];
+			TotalVolume += BigVol;
+			UpdateMergeTarget(BigIdx, BigVol);
+			if (BigVol < MinVolume)
+			{
+				checkSlow(!ToRemove.Contains(BigIdx));
+				ToRemove.Add(BigIdx);
+			}
+			else
+			{
+				bRemoveMergeTarget = false;
+			}
+		}
+
+		void TransferGroup(FRemoveGroup& SmallGroup, TMap<int32, int32>& GeomIdxToRemoveGroupIdx, int32 NewIdx)
+		{
+			for (int32 RmIdx : SmallGroup.ToRemove)
+			{
+				checkSlow(!ToRemove.Contains(RmIdx));
+				ToRemove.Add(RmIdx);
+				GeomIdxToRemoveGroupIdx[RmIdx] = NewIdx;
+			}
+			if (!SmallGroup.bRemoveMergeTarget)
+			{
+				checkSlow(!ToRemove.Contains(SmallGroup.MergeTo));
+				ToRemove.Add(SmallGroup.MergeTo);
+				GeomIdxToRemoveGroupIdx[SmallGroup.MergeTo] = NewIdx;
+			}
+			TotalVolume += SmallGroup.TotalVolume;
+			SmallGroup = FRemoveGroup(); // clear old group
+		}
+
+		bool IsGroupSmall(double MinVolume)
+		{
+			return TotalVolume < MinVolume;
+		}
+	};
+	TMap<int32, double> GeomToVol;
+	TMap<int32, int32> GeomIdxToRemoveGroupIdx;
+	TArray<FRemoveGroup> RemoveGroups;
+	TSet<int32> TooSmalls;
+	TSet<int32> CanMerge;
+
+	for (int32 TransformIdx : TransformIndices)
+	{
+		int32 GeomIdx = Collection.TransformToGeometryIndex[TransformIdx];
+		if (GeomIdx > -1)
+		{
+			CanMerge.Add(GeomIdx);
+			GeomToVol.Add(GeomIdx, Volumes[TransformIdx]);
+		}
+	}
+	for (int32 TransformIdx : SmallTransformIndices)
+	{
+		int32 GeomIdx = Collection.TransformToGeometryIndex[TransformIdx];
+		if (GeomIdx > -1)
+		{
+			TooSmalls.Add(GeomIdx);
+			GeomToVol.Add(GeomIdx, Volumes[TransformIdx]);
+		}
+		else
+		{
+			ensureMsgf(false, TEXT("Cannot merge bones that have no geometry attached"));
+		}
+	}
+
+	for (int32 SmallIdx : TooSmalls)
+	{
+		int32* SmallRemoveGroupIdx = GeomIdxToRemoveGroupIdx.Find(SmallIdx);
+		if (SmallRemoveGroupIdx)
+		{
+			if (RemoveGroups[*SmallRemoveGroupIdx].TotalVolume >= MinVolume)
+			{
+				continue;
+			}
+		}
+
+		const TSet<int32>& Prox = Proximity[SmallIdx];
+		double BiggestVol = 0;
+		int32 BiggestVolIdx = INDEX_NONE;
+		for (int32 NbrIdx : Prox)
+		{
+			if (CanMerge.Contains(NbrIdx))
+			{
+				double Vol = GeomToVol[NbrIdx];
+				if (Vol > BiggestVol)
+				{
+					BiggestVol = Vol;
+					BiggestVolIdx = NbrIdx;
+				}
+			}
+		}
+		if (BiggestVolIdx == INDEX_NONE)
+		{
+			continue;
+		}
+
+		if (SmallRemoveGroupIdx)
+		{
+			int32 OldSGIdx = *SmallRemoveGroupIdx;
+			int32* BigRemoveGroupIdx = GeomIdxToRemoveGroupIdx.Find(BiggestVolIdx);
+			if (BigRemoveGroupIdx)
+			{
+				int32 BigRGIdx = *BigRemoveGroupIdx;
+				if (OldSGIdx != BigRGIdx)
+				{
+					RemoveGroups[BigRGIdx].TransferGroup(RemoveGroups[OldSGIdx], GeomIdxToRemoveGroupIdx, BigRGIdx);
+					checkSlow(GeomIdxToRemoveGroupIdx.FindKey(OldSGIdx) == nullptr);
+				}
+			}
+			else
+			{
+				RemoveGroups[OldSGIdx].AddBig(GeomToVol, MinVolume, BiggestVolIdx);
+				checkSlow(!GeomIdxToRemoveGroupIdx.Contains(BiggestVolIdx));
+				GeomIdxToRemoveGroupIdx.Add(BiggestVolIdx, OldSGIdx);
+			}
+		}
+		else
+		{
+			int32* BigRemoveGroupIdx = GeomIdxToRemoveGroupIdx.Find(BiggestVolIdx);
+			if (BigRemoveGroupIdx)
+			{
+				int32 BigRGIdx = *BigRemoveGroupIdx;
+				RemoveGroups[BigRGIdx].AddSmall(GeomToVol, SmallIdx);
+				checkSlow(!GeomIdxToRemoveGroupIdx.Contains(SmallIdx));
+				GeomIdxToRemoveGroupIdx.Add(SmallIdx, BigRGIdx);
+			}
+			else
+			{
+				int32 RemoveGroupIdx = RemoveGroups.Emplace(GeomToVol, MinVolume, SmallIdx, BiggestVolIdx);
+				checkSlow(!GeomIdxToRemoveGroupIdx.Contains(SmallIdx));
+				checkSlow(!GeomIdxToRemoveGroupIdx.Contains(BiggestVolIdx));
+				GeomIdxToRemoveGroupIdx.Add(SmallIdx, RemoveGroupIdx);
+				GeomIdxToRemoveGroupIdx.Add(BiggestVolIdx, RemoveGroupIdx);
+			}
+		}
+	}
+
+	TArray<int32> AllRemoveIndices;
+	TArray<int32> AllUpdateIndices;
+
+	for (FRemoveGroup& Group : RemoveGroups)
+	{
+		if (!Group.IsValid())
+		{
+			continue;
+		}
+		if (Group.bRemoveMergeTarget)
+		{
+			Group.ToRemove.RemoveSingle(Group.MergeTo);
+		}
+		for (int32 RmIdx : Group.ToRemove)
+		{
+			AllRemoveIndices.Add(Collection.TransformIndex[RmIdx]);
+		}
+		AllUpdateIndices.Add(Collection.TransformIndex[Group.MergeTo]);
+	}
+
+	AllRemoveIndices.Sort();
+	AllUpdateIndices.Sort();
+	
+	FDynamicMeshCollection RemoveCollection(&Collection, AllRemoveIndices, CellsToWorld, true);
+	FDynamicMeshCollection UpdateCollection(&Collection, AllUpdateIndices, CellsToWorld, true);
+
+	TMap<int32, int32> GeoIdxToRmMeshIdx;
+	for (int32 RmMeshIdx = 0; RmMeshIdx < RemoveCollection.Meshes.Num(); RmMeshIdx++)
+	{
+		int32 TransformIdx = RemoveCollection.Meshes[RmMeshIdx].TransformIndex;
+		GeoIdxToRmMeshIdx.Add(
+			Collection.TransformToGeometryIndex[TransformIdx],
+			RmMeshIdx
+		);
+	}
+
+	using FMeshData = UE::PlanarCut::FDynamicMeshCollection::FMeshData;
+
+	for (int32 UpMeshIdx = 0; UpMeshIdx < UpdateCollection.Meshes.Num(); UpMeshIdx++)
+	{
+		FMeshData& UpMeshData = UpdateCollection.Meshes[UpMeshIdx];
+		FTransform3d UpMeshXF = (FTransform3d)UpMeshData.ToCollection;
+		int32 UpGeoIdx = Collection.TransformToGeometryIndex[UpMeshData.TransformIndex];
+		FRemoveGroup& Group = RemoveGroups[GeomIdxToRemoveGroupIdx[UpGeoIdx]];
+		if (!ensure(Group.IsValid()))
+		{
+			continue;
+		}
+		FDynamicMeshEditor MeshEditor(&UpMeshData.AugMesh);
+		for (int32 RmGeoIdx : Group.ToRemove)
+		{
+			FMeshData& RmMeshData = RemoveCollection.Meshes[GeoIdxToRmMeshIdx[RmGeoIdx]];
+			FTransform3d RmMeshXF = (FTransform3d)RmMeshData.ToCollection;
+			if (bUnionJoinedPieces)
+			{
+				FMeshBoolean Boolean(&UpMeshData.AugMesh, &RmMeshData.AugMesh, &UpMeshData.AugMesh, FMeshBoolean::EBooleanOp::Union);
+				Boolean.bWeldSharedEdges = false;
+				Boolean.bSimplifyAlongNewEdges = true;
+				Boolean.Compute();
+			}
+			else
+			{
+				FMeshIndexMappings IndexMaps_Unused;
+				MeshEditor.AppendMesh(&RmMeshData.AugMesh, IndexMaps_Unused,
+					[&UpMeshXF, &RmMeshXF](int32 Unused, const FVector3d& Pos)
+					{
+						return UpMeshXF.InverseTransformPosition(RmMeshXF.TransformPosition(Pos));
+					},
+					[&UpMeshXF, &RmMeshXF](int32 Unused, const FVector3d& Normal)
+					{
+						return UpMeshXF.InverseTransformNormal(RmMeshXF.TransformNormal(Normal));
+					});
+			}
+		}
+	}
+
+	UpdateCollection.UpdateAllCollections(Collection);
+
+	for (FRemoveGroup& Group : RemoveGroups)
+	{
+		if (!Group.IsValid())
+		{
+			continue;
+		}
+		int32 MergeTransformIdx = Collection.TransformIndex[Group.MergeTo];
+		TArray<int32> Children;
+		for (int32 RmIdx : Group.ToRemove)
+		{
+			int32 RmTransformIdx = Collection.TransformIndex[RmIdx];
+			for (int32 ChildIdx : Collection.Children[RmTransformIdx])
+			{
+				Children.Add(ChildIdx);
+			}
+		}
+		if (Children.Num())
+		{
+			GeometryCollectionAlgo::ParentTransforms(&Collection, MergeTransformIdx, Children);
+		}
+	}
+
+	// remove transforms for all geometry that was merged in
+	Collection.RemoveElements(FGeometryCollection::TransformGroup, AllRemoveIndices);
+
+	return INDEX_NONE; // TODO: consider tracking smallest index of updated groups?  but no reason to do so currently
+}
+
+
 void RecomputeNormalsAndTangents(bool bOnlyTangents, FGeometryCollection& Collection, const TArrayView<const int32>& TransformIndices,
 								 bool bOnlyOddMaterials, const TArrayView<const int32>& WhichMaterials)
 {
diff --git a/Engine/Plugins/Experimental/PlanarCutPlugin/Source/PlanarCut/Public/PlanarCut.h b/Engine/Plugins/Experimental/PlanarCutPlugin/Source/PlanarCut/Public/PlanarCut.h
--- a/Engine/Plugins/Experimental/PlanarCutPlugin/Source/PlanarCut/Public/PlanarCut.h
+++ b/Engine/Plugins/Experimental/PlanarCutPlugin/Source/PlanarCut/Public/PlanarCut.h
@@ -208,6 +208,58 @@ int32 PLANARCUT_API CutMultipleWithMultiplePlanes(
 	bool bSetDefaultInternalMaterialsFromCollection = true
 );
 
+
+/**
+ * Populate an array of transform indices w/ those that are smaller than a threshold volume
+ *
+ * @param Collection			The collection to be processed
+ * @param TransformIndices		The transform indices to process, or empty if all should be processed
+ * @param OutVolumes			Output array, to be filled w/ volumes of geometry; 1:1 w/ TransformIndices array
+ * @param ScalePerDimension		Scale to apply per dimension (e.g. 1/100 converts volume from centimeters^3 to meters^3)
+ */
+void PLANARCUT_API FindBoneVolumes(
+	FGeometryCollection& Collection,
+	const TArrayView<const int32>& TransformIndices,
+	TArray<double>& OutVolumes,
+	double ScalePerDimension = .01
+);
+
+/**
+ * Populate an array of transform indices w/ those that are smaller than a threshold volume
+ *
+ * @param Collection			The collection to be processed
+ * @param TransformIndices		The transform indices to process, or empty if all should be processed
+ * @param Volumes				Volumes of geometry; 1:1 w/ TransformIndices array
+ * @param MinVolume				Geometry smaller than this quantity will be chosen
+ * @param OutSmallBones			Output array, to be filled with transform indices for small pieces of geometry
+ */
+void PLANARCUT_API FindSmallBones(
+	FGeometryCollection& Collection,
+	const TArrayView<const int32>& TransformIndices,
+	const TArrayView<const double>& Volumes,
+	double MinVolume,
+	TArray<int32>& OutSmallBones
+);
+
+/**
+ * Merge chosen geometry into neighboring geometry.  Geometry will be merged into its largest neighboring geometry (if any).
+ *
+ * @param Collection			The collection to be processed
+ * @param TransformIndices		The transform indices to process, or empty if all should be processed
+ * @param Volumes				Volumes of geometry; 1:1 w/ TransformIndices array
+ * @param MinVolume				If merged small geometry is larger than this, it will not require further merging
+ * @param SmallTransformIndices	Transformed indices of pieces that we want to merge
+ * @param bUnionJoinedPieces	Try to 'union' the merged pieces, removing internal triangles and connecting the shared cut boundary
+ */
+int32 PLANARCUT_API MergeBones(
+	FGeometryCollection& Collection,
+	const TArrayView<const int32>& TransformIndices,
+	const TArrayView<const double>& Volumes,
+	double MinVolume,
+	const TArrayView<const int32>& SmallTransformIndices,
+	bool bUnionJoinedPieces
+);
+
 /**
  * Recompute normals and tangents of selected geometry, optionally restricted to faces with odd or given material IDs (i.e. to target internal faces)
  *
diff --git a/Engine/Source/Runtime/Experimental/Chaos/Private/GeometryCollection/GeometryCollectionAlgo.cpp b/Engine/Source/Runtime/Experimental/Chaos/Private/GeometryCollection/GeometryCollectionAlgo.cpp
--- a/Engine/Source/Runtime/Experimental/Chaos/Private/GeometryCollection/GeometryCollectionAlgo.cpp
+++ b/Engine/Source/Runtime/Experimental/Chaos/Private/GeometryCollection/GeometryCollectionAlgo.cpp
@@ -794,6 +794,13 @@ namespace GeometryCollectionAlgo
 		if (NewNumFaces > OldNumFaces)
 		{
 			GeometryCollection->AddElements(NewNumFaces - OldNumFaces, FGeometryCollection::FacesGroup);
+			// fill in end faces with dummy values, rather than leaving them uninitialized (to avoid breaking things on later reordering call)
+			int LastVertex = OldNumVertices > 0 ? OldNumVertices - 1 : 0;
+			FIntVector EndFace(LastVertex, LastVertex, LastVertex);
+			for (int Idx = OldNumFaces; Idx < NewNumFaces; Idx++)
+			{
+				GeometryCollection->Indices[Idx] = EndFace;
+			}
 		}
 		if (NewNumVertices > OldNumVertices)
 		{
@@ -880,18 +887,18 @@ namespace GeometryCollectionAlgo
 		}
 
 		// remove trailing elements if needed
-		if (NewNumFaces < OldNumFaces)
-		{
-			TArray<int32> ToDelete;
-			AddRange(ToDelete, NewNumFaces, OldNumFaces);
-			GeometryCollection->RemoveElements(FGeometryCollection::FacesGroup, ToDelete);
-		}
 		if (NewNumVertices < OldNumVertices)
 		{
 			TArray<int32> ToDelete;
 			AddRange(ToDelete, NewNumVertices, OldNumVertices);
 			GeometryCollection->RemoveElements(FGeometryCollection::VerticesGroup, ToDelete);
 		}
+		if (NewNumFaces < OldNumFaces)
+		{
+			TArray<int32> ToDelete;
+			AddRange(ToDelete, NewNumFaces, OldNumFaces);
+			GeometryCollection->RemoveElements(FGeometryCollection::FacesGroup, ToDelete);
+		}
 
 		// TODO: can remove these ensure()s after this has been tested some more
 		ensure(GeometryCollection->HasContiguousFaces());
