Add a custom hole filler for fracture (automatically called after fracture Boolean ops)

#rnx

diff --git a/Engine/Plugins/Experimental/PlanarCutPlugin/Source/PlanarCut/Private/GeometryMeshConversion.cpp b/Engine/Plugins/Experimental/PlanarCutPlugin/Source/PlanarCut/Private/GeometryMeshConversion.cpp
--- a/Engine/Plugins/Experimental/PlanarCutPlugin/Source/PlanarCut/Private/GeometryMeshConversion.cpp
+++ b/Engine/Plugins/Experimental/PlanarCutPlugin/Source/PlanarCut/Private/GeometryMeshConversion.cpp
@@ -15,6 +15,7 @@
 #include "FrameTypes.h"
 #include "Polygon2.h"
 #include "CompGeom/PolygonTriangulation.h"
+#include "Quaternion.h"
 
 #include "GeometryCollection/GeometryCollectionAlgo.h"
 #include "GeometryCollection/GeometryCollectionClusteringUtility.h"
@@ -597,6 +598,392 @@ namespace AugmentedDynamicMesh
 			}
 		}
 	}
+
+	/**
+	 * Fill holes in non-welded meshes, w/ UVs that attempt to continue one of the local UV islands.
+	 * Intended to clean holes erroneously left by FMeshBoolean (due to floating point error)
+	 * 
+	 * Note this algo tries to continue UVs of internal faces only, where the valid UV areas aren't already
+	 * pre-defined to match an artist-authored texture, but are automatically generated by simple projection
+	 * per connected component.
+	 * Hole filling with 'external' materials and UVs could risk having very stretched triangles in UV space.
+	 * The hope is that holes are rare and small enough that it's ok to just use internal materials.
+	 * 
+	 * @param Mesh Fill holes on this mesh
+	 * @param CandidateEdges Only consider holes that touch these edges
+	 */
+	void FillHoles(FDynamicMesh3& Mesh, const TSet<int32>& CandidateEdges, double MinHoleArea = KINDA_SMALL_NUMBER)
+	{
+		/// 1. Build a spatial hash of boundary vertices, so we can identify holes even though the mesh is not welded
+		
+		double SnapDistance = 1e-05;
+		TPointHashGrid3d<int32> VertHash(SnapDistance * 10, -1);
+		TMap<int32, int32> CoincidentVerticesMap; // map every vertex in an overlapping cluster to a canonical single vertex for that cluster
+		TSet<int32> HashedVertices; // track what's already processed
+
+		auto AddVertexToHash = [&Mesh, SnapDistance, &VertHash, &CoincidentVerticesMap, &HashedVertices](int32 VID)
+		{
+			if (HashedVertices.Contains(VID))
+			{
+				return;
+			}
+
+			HashedVertices.Add(VID);
+
+			FVector3d VPos = Mesh.GetVertex(VID);
+			TPair<int32, double> Res = VertHash.FindNearestInRadius(VPos, SnapDistance, [&Mesh, &VPos](const int& VID)->double { return FVector3d::DistSquared(VPos, Mesh.GetVertex(VID)); });
+			if (Res.Key != INDEX_NONE)
+			{
+				int32 MapTo = Res.Key;
+				int32* WasMapped = CoincidentVerticesMap.Find(MapTo);
+				if (WasMapped)
+				{
+					MapTo = *WasMapped;
+				}
+				CoincidentVerticesMap.Add(VID, MapTo);
+			}
+
+			VertHash.InsertPointUnsafe(VID, VPos);
+		};
+
+		for (int32 EID : CandidateEdges)
+		{
+			FDynamicMesh3::FEdge Edge = Mesh.GetEdge(EID);
+			if (Edge.Tri.B == -1) // only consider boundary edges
+			{
+				AddVertexToHash(Edge.Vert.A);
+				AddVertexToHash(Edge.Vert.B);
+			}
+		}
+
+		/// 2. Find the edges w/ no paired reverse edge (i.e. the potential hole boundary edges)
+
+		using FPair = TPair<int32, int32>;
+		TMultiMap<int32, FPair> OpenEdges; // FPair contains: (Second vertex ID on edge, Edge ID)
+		auto CanonicalVID = [&CoincidentVerticesMap](int32 VID)
+		{
+			int32* Mapped = CoincidentVerticesMap.Find(VID);
+			return Mapped ? *Mapped : VID;
+		};
+		// remove an edge from OpenEdges, assuming the vertices are already passed through CanonicalVID()
+		auto RemoveCanonEdge = [&OpenEdges](int32 A, int32 B)
+		{
+			int32 NumRemovedPairs = 0;
+			for (TMultiMap<int32, FPair>::TKeyIterator It = OpenEdges.CreateKeyIterator(A); It; ++It)
+			{
+				if (It.Value().Key == B)
+				{
+					It.RemoveCurrent();
+					return true;
+				}
+			}
+			return false;
+		};
+		auto AddEdge = [&OpenEdges, &CanonicalVID, &RemoveCanonEdge](int32 A, int32 B, int32 EID)
+		{
+			A = CanonicalVID(A);
+			B = CanonicalVID(B);
+			if (A == B) // skip edges where the vertices were too close
+			{
+				return;
+			}
+			if (!RemoveCanonEdge(B, A))
+			{
+				OpenEdges.Add(A, FPair(B, EID));
+			}
+		};
+		for (int32 CandEID : CandidateEdges)
+		{
+			FDynamicMesh3::FEdge Edge = Mesh.GetEdge(CandEID);
+			if (Edge.Tri[1] == FDynamicMesh3::InvalidID)
+			{
+				int32 TID = Edge.Tri[0];
+				FIndex2i EdgeV = Mesh.GetOrientedBoundaryEdgeV(CandEID);
+				AddEdge(EdgeV.A, EdgeV.B, CandEID);
+			}
+		}
+
+		if (OpenEdges.Num() < 3)
+		{
+			return;
+		}
+
+		/// 3. Walk the open boundary edges to find the loops that define holes
+
+		TArray<int32> AllBoundaryVerts; // VIDs for all hole boundaries (w/ Num IDs at the start of each span)
+		TArray<int32> AllBoundaryEdges; // EIDs for hole boundaries (w/ same layout as AllBoundaryVerts)
+		int32 NumHoles = 0;
+		while (true)
+		{
+			const auto& EdgeItr = OpenEdges.CreateConstIterator();
+			if (!EdgeItr)
+			{
+				break;
+			}
+			int32 Start = EdgeItr.Key();
+			FPair WalkPair = EdgeItr.Value();
+			int32 Walk = WalkPair.Key;
+			int32 BdryStartIdx = AllBoundaryVerts.Add(-1);
+			int32 BdryEdgeStartIdx = AllBoundaryEdges.Add(-1);
+			checkSlow(BdryEdgeStartIdx == BdryStartIdx);
+			AllBoundaryVerts.Add(Start);
+			AllBoundaryEdges.Add(WalkPair.Value);
+			RemoveCanonEdge(Start, Walk);
+			while (Walk != Start) // this will either loop back to start or dead end (note every non-dead-end iter removes an edge from OpenEdges)
+			{
+				AllBoundaryVerts.Add(Walk);
+				FPair* Next = OpenEdges.Find(Walk);
+				if (!Next)
+				{
+					// dead-ended chain of edges, discard the vertices
+					AllBoundaryVerts.SetNum(BdryStartIdx);
+					AllBoundaryEdges.SetNum(BdryStartIdx);
+					break;
+				}
+				int32 NextVal = Next->Key;
+				AllBoundaryEdges.Add(Next->Value);
+				RemoveCanonEdge(Walk, NextVal);
+				Walk = NextVal;
+			}
+			int32 Count = AllBoundaryVerts.Num() - BdryStartIdx - 1;
+			if (Count < 3) // failed to add a boundary, try again next loop
+			{
+				AllBoundaryVerts.SetNum(BdryStartIdx);
+				AllBoundaryEdges.SetNum(BdryStartIdx);
+				continue;
+			}
+			
+			AllBoundaryVerts[BdryStartIdx] = Count; // record length of added chain
+			AllBoundaryEdges[BdryStartIdx] = Count;
+			NumHoles++;
+		}
+		if (NumHoles == 0)
+		{
+			return;
+		}
+		check(AllBoundaryVerts.Num() == AllBoundaryEdges.Num());
+
+		/// Tag vertex connected components
+
+		FDisjointSet VertComponents(Mesh.MaxVertexID());
+		for (FIndex3i Tri : Mesh.TrianglesItr())
+		{
+			VertComponents.Union(Tri.A, Tri.B);
+			VertComponents.Union(Tri.B, Tri.C);
+		}
+
+		FDynamicMeshMaterialAttribute* MaterialIDs = Mesh.Attributes()->GetMaterialID();
+
+		/// 4. Compute stats on the holes: Their areas, and the neighboring component that best matches the hole normal
+
+		TArray<int32> HoleComponents; // which connected component to connect the new hole geometry to
+		HoleComponents.Init(-1, NumHoles);
+		TArray<double> HoleAreas;
+		HoleAreas.Init(0, NumHoles);
+		TArray<FVector3d> HoleNormals;
+		HoleNormals.Init(FVector3d::ZeroVector, NumHoles);
+		for (int32 BIdx = 0, CurBdry = 0; BIdx < AllBoundaryVerts.Num(); BIdx += AllBoundaryVerts[BIdx] + 1, CurBdry++)
+		{
+			int32 ChainLen = AllBoundaryVerts[BIdx];
+			int32 EndIdx = EndIdx = BIdx + ChainLen + 1;
+			check(ChainLen > 2 && EndIdx <= AllBoundaryVerts.Num());
+
+			TMap<int32, FVector3d> ComponentToNormalMap;
+			FVector3d HoleNormal(0, 0, 0), LastEdge(0, 0, 0);
+			double Area = 0;
+			for (int32 LastIdx = EndIdx - 1, Idx = BIdx + 1; Idx < EndIdx; LastIdx = Idx++)
+			{
+				int32 V0 = AllBoundaryVerts[LastIdx];
+				int32 V1 = AllBoundaryVerts[Idx];
+				FVector3d P0 = Mesh.GetVertex(V0);
+				FVector3d P1 = Mesh.GetVertex(V1);
+				FVector3d Edge = P1 - P0;
+				FVector3d HoleTriNormal = Edge.Cross(LastEdge);
+				Area += HoleTriNormal.Length() * .5;
+				HoleNormal += HoleTriNormal;
+				LastEdge = Edge;
+
+				int32 EID01 = AllBoundaryEdges[LastIdx];
+				FDynamicMesh3::FEdge Edge01 = Mesh.GetEdge(EID01);
+				int32 TID01 = Edge01.Tri.A;
+				int32 MID01 = MaterialIDs->GetValue(TID01);
+				if (MID01 >= 0 && (MID01 % 2) == 0)
+				{
+					continue; // skip "outside" materials, where UVs are less predictable
+				}
+				int32 Component = (int32)VertComponents.Find(Edge01.Vert.A);
+				FVector3d Normal = Mesh.GetTriNormal(TID01);
+				FVector3d* AccNormal = ComponentToNormalMap.Find(Component);
+				if (AccNormal)
+				{
+					(*AccNormal) += Normal;
+				}
+				else
+				{
+					ComponentToNormalMap.Add(Component, Normal);
+				}
+			}
+			int32 BestComponent = -1;
+			double BestScore = -2; // Score is dot product of component avg triangle normal and hole normal
+			bool bHasNormal = HoleNormal.Normalize(DBL_EPSILON);
+			for (TPair<int32, FVector3d>& IDNormal : ComponentToNormalMap)
+			{
+				double Score;
+				if (IDNormal.Value.Normalize(DBL_EPSILON))
+				{
+					Score = HoleNormal.Dot(IDNormal.Value);
+				}
+				else
+				{
+					Score = -1;
+				}
+				if (Score > BestScore)
+				{
+					BestScore = Score;
+					BestComponent = IDNormal.Key;
+				}
+			}
+
+			HoleAreas[CurBdry] = Area;
+			if (!bHasNormal && BestComponent > -1)
+			{
+				HoleNormal = ComponentToNormalMap[BestComponent];
+			}
+			if (BestComponent == -1)
+			{
+				ensure(false);
+				BestComponent = VertComponents.Find(AllBoundaryVerts[BIdx + 1]);
+			}
+			HoleComponents[CurBdry] = BestComponent;
+			HoleNormals[CurBdry] = HoleNormal;
+		}
+
+		/// 5. Fill holes with new triangles
+
+		int32 NumUVs = NumEnabledUVChannels(Mesh);
+
+		TArray<int32> UseVIDs; // vertex IDs for the current hole
+		for (int32 BIdx = 0, CurBdry = 0; BIdx < AllBoundaryVerts.Num(); BIdx += AllBoundaryVerts[BIdx] + 1, CurBdry++)
+		{
+			int32 ChainLen = AllBoundaryVerts[BIdx];
+			int32 EndIdx = EndIdx = BIdx + ChainLen + 1;
+			check(ChainLen > 2 && EndIdx <= AllBoundaryVerts.Num());
+			if (HoleAreas[CurBdry] < MinHoleArea)
+			{
+				continue;
+			}
+
+			// find the hole material and edge matching the target vertex component
+			int32 HoleMaterial = -1;
+			FIndex2i HoleEdgeV; // reference edge for the hole
+			for (int32 LastIdx = EndIdx - 1, Idx = BIdx + 1; Idx < EndIdx; LastIdx = Idx++)
+			{
+				int32 V0 = AllBoundaryVerts[LastIdx];
+				int32 V1 = AllBoundaryVerts[Idx];
+				int32 EID01 = AllBoundaryEdges[LastIdx];
+				FDynamicMesh3::FEdge Edge = Mesh.GetEdge(EID01);
+				int32 Component01 = (int32)VertComponents.Find((uint32)Edge.Vert.A);
+				if (Component01 == HoleComponents[CurBdry])
+				{
+					check(Edge.Vert.A != INDEX_NONE);
+					HoleMaterial = MaterialIDs->GetValue(Edge.Tri.A);
+					HoleEdgeV = Mesh.GetOrientedBoundaryEdgeV(EID01);
+					break;
+				}
+			}
+
+			// find a basis for UV projection (per UV channel)
+			FVector3d N = HoleNormals[CurBdry];
+			FVector3d Origin = Mesh.GetVertex(HoleEdgeV.A);
+			FVector3d T[8], B[8];
+			{
+				FVector3d Edge = Mesh.GetVertex(HoleEdgeV.B) - Origin;
+				float EdgeLen = (float)Edge.Length();
+				for (int32 UVIdx = 0; UVIdx < NumUVs; UVIdx++)
+				{
+					FVector2f UVA, UVB;
+					GetUV(Mesh, HoleEdgeV.A, UVA, UVIdx);
+					GetUV(Mesh, HoleEdgeV.B, UVB, UVIdx);
+					FVector2f UVEdge = UVB - UVA;
+					float UVEdgeLen = UVEdge.Length();
+					float UVScale = UVEdgeLen / EdgeLen;
+					double Angle = (double)FMath::Atan2(UVEdge.Y, UVEdge.X);
+					UE::Geometry::FQuaterniond RotT(N, -Angle, false);
+					UE::Geometry::FQuaterniond RotB(N, FMathd::HalfPi - Angle, false);
+					T[UVIdx] = RotT * Edge * UVScale;
+					B[UVIdx] = RotB * Edge * UVScale;
+
+				}
+			}
+
+			// Indices of vertices to be used for the hole fill
+			UseVIDs.Init(-1, ChainLen);
+
+			// Find the vertices that match the target component -- we can re-use these directly
+			for (int32 LastIdx = EndIdx - 1, Idx = BIdx + 1; Idx < EndIdx; LastIdx = Idx++)
+			{
+				int32 V0 = AllBoundaryVerts[LastIdx];
+				int32 V1 = AllBoundaryVerts[Idx];
+
+				int32 EID01 = AllBoundaryEdges[LastIdx];
+				
+				int32 C0 = LastIdx - BIdx - 1;
+				int32 C1 = Idx - BIdx - 1;
+				FIndex2i EdgeV = Mesh.GetOrientedBoundaryEdgeV(EID01);
+				if (HoleComponents[CurBdry] == (int32)VertComponents.Find(EdgeV.A))
+				{
+					UseVIDs[C0] = EdgeV.A;
+					UseVIDs[C1] = EdgeV.B;
+				}
+			}
+
+			// Copy color + tangents from the reference vertex (TODO: consider an average instead?)
+			FVector3f HoleColor;
+			FVector3f HoleTanU, HoleTanV;
+			GetTangent(Mesh, HoleEdgeV.A, HoleTanU, HoleTanV);
+			HoleColor = Mesh.GetVertexColor(HoleEdgeV.A);
+			
+			// Create new vertices for all that weren't already assigned to an existing vertex
+			for (int32 Idx = BIdx + 1; Idx < EndIdx; Idx++)
+			{
+				int32 C = Idx - BIdx - 1;
+				if (UseVIDs[C] != -1)
+				{
+					continue;
+				}
+
+				int32 V = AllBoundaryVerts[Idx];
+				FVector3d P = Mesh.GetVertex(V);
+				UE::Geometry::FVertexInfo Info(P, HoleNormals[CurBdry], HoleColor);
+				int32 NewV = Mesh.AppendVertex(Info);
+				SetTangent(Mesh, NewV, HoleNormals[CurBdry], HoleTanU, HoleTanV);
+				for (int32 UVIdx = 0; UVIdx < NumUVs; UVIdx++)
+				{
+					FVector3d Diff = P - Origin;
+					FVector2f UV((float)Diff.Dot(T[UVIdx]), (float)Diff.Dot(B[UVIdx]));
+					SetUV(Mesh, NewV, UV, UVIdx);
+				}
+				UseVIDs[C] = NewV;
+			}
+
+			int32 V0 = UseVIDs[0];
+			int32 V1, V2 = UseVIDs[1];
+			FVector3d P0 = Mesh.GetVertex(V0);
+			FVector3d P1, P2 = Mesh.GetVertex(V2);
+			for (int32 C1 = 1, C2 = 2; C2 < ChainLen; C1 = C2++)
+			{
+				V1 = V2;
+				V2 = UseVIDs[C2];
+				P1 = P2;
+				P2 = Mesh.GetVertex(V2);
+
+				int32 TID = Mesh.AppendTriangle(V0, V2, V1);
+				MaterialIDs->SetValue(TID, HoleMaterial);
+				SetVisibility(Mesh, TID, true);
+			}
+		}
+	}
+
 } // namespace AugmentedDynamicMesh
 
 
@@ -1859,6 +2246,7 @@ int32 FDynamicMeshCollection::CutWithMultiplePlanes(
 				Boolean.bSimplifyAlongNewEdges = true;
 				Boolean.PreserveUVsOnlyForMesh = 0; // slight warping of the autogenerated cell UVs generally doesn't matter
 				Boolean.bWeldSharedEdges = false;
+				Boolean.bTrackAllNewEdges = true;
 				if (!Boolean.Compute())
 				{
 					// TODO: do something about failure cases?  e.g. try auto-filling small holes?
@@ -1870,6 +2258,7 @@ int32 FDynamicMeshCollection::CutWithMultiplePlanes(
 					bKeepResults = false;
 					break;
 				}
+				AugmentedDynamicMesh::FillHoles(BoolResults[CellIdx]->AugMesh, Boolean.AllNewEdges);
 			}
 
 			if (bKeepResults)
@@ -2051,12 +2440,14 @@ int32 FDynamicMeshCollection::CutWithCellMeshes(const FInternalSurfaceMaterials&
 					Boolean.bSimplifyAlongNewEdges = true;
 					Boolean.PreserveUVsOnlyForMesh = 0; // slight warping of the autogenerated cell UVs generally doesn't matter
 					Boolean.bWeldSharedEdges = false;
+					Boolean.bTrackAllNewEdges = true;
 					if (!Boolean.Compute())
 					{
 						// TODO: do something about failure cases?  e.g. try auto-filling small holes?
 						// note: failure cases won't be detected at all unless we weld edges,
 						//       which will require re-working how tangents are carried through
 					}
+					AugmentedDynamicMesh::FillHoles(AugBoolResult, Boolean.AllNewEdges);
 				}
 			}, EParallelForFlags::None);
 
@@ -2263,7 +2654,6 @@ bool FDynamicMeshCollection::SplitIslands(FDynamicMesh3& Source, TArray<FDynamic
 	{
 		VertComponents.Union(Tri.A, Tri.B);
 		VertComponents.Union(Tri.B, Tri.C);
-		VertComponents.Union(Tri.C, Tri.A);
 	}
 
 	bool bWasSplit = FDynamicMeshEditor::SplitMesh(&Source, SeparatedMeshes, [&Source, &VertComponents](int TID)
diff --git a/Engine/Plugins/Runtime/GeometryProcessing/Source/DynamicMesh/Private/Operations/MeshBoolean.cpp b/Engine/Plugins/Runtime/GeometryProcessing/Source/DynamicMesh/Private/Operations/MeshBoolean.cpp
--- a/Engine/Plugins/Runtime/GeometryProcessing/Source/DynamicMesh/Private/Operations/MeshBoolean.cpp
+++ b/Engine/Plugins/Runtime/GeometryProcessing/Source/DynamicMesh/Private/Operations/MeshBoolean.cpp
@@ -673,6 +673,22 @@ bool FMeshBoolean::Compute()
 			bool bWeldSuccess = MergeEdges(IndexMaps, CutMesh, CutBoundaryEdges, AllVIDMatches);
 			bSuccess = bSuccess && bWeldSuccess;
 		}
+		else
+		{
+			CreatedBoundaryEdges = CutBoundaryEdges[0];
+			for (int OldMeshEID : CutBoundaryEdges[1])
+			{
+				if (!CutMesh[1]->IsEdge(OldMeshEID))
+				{
+					ensure(false);
+					continue;
+				}
+				FIndex2i OtherEV = CutMesh[1]->GetEdgeV(OldMeshEID);
+				int MappedEID = Result->FindEdge(IndexMaps.GetNewVertex(OtherEV.A), IndexMaps.GetNewVertex(OtherEV.B));
+				checkSlow(Result->IsBoundaryEdge(MappedEID));
+				CreatedBoundaryEdges.Add(MappedEID);
+			}
+		}
 	}
 	else
 	{
