CIS fix for warning on redundant statement in fracture hole filling function

#rnx

diff --git a/Engine/Plugins/Experimental/PlanarCutPlugin/Source/PlanarCut/Private/GeometryMeshConversion.cpp b/Engine/Plugins/Experimental/PlanarCutPlugin/Source/PlanarCut/Private/GeometryMeshConversion.cpp
--- a/Engine/Plugins/Experimental/PlanarCutPlugin/Source/PlanarCut/Private/GeometryMeshConversion.cpp
+++ b/Engine/Plugins/Experimental/PlanarCutPlugin/Source/PlanarCut/Private/GeometryMeshConversion.cpp
@@ -785,7 +785,7 @@ namespace AugmentedDynamicMesh
 		for (int32 BIdx = 0, CurBdry = 0; BIdx < AllBoundaryVerts.Num(); BIdx += AllBoundaryVerts[BIdx] + 1, CurBdry++)
 		{
 			int32 ChainLen = AllBoundaryVerts[BIdx];
-			int32 EndIdx = EndIdx = BIdx + ChainLen + 1;
+			int32 EndIdx = BIdx + ChainLen + 1;
 			check(ChainLen > 2 && EndIdx <= AllBoundaryVerts.Num());
 
 			TMap<int32, FVector3d> ComponentToNormalMap;
@@ -866,7 +866,7 @@ namespace AugmentedDynamicMesh
 		for (int32 BIdx = 0, CurBdry = 0; BIdx < AllBoundaryVerts.Num(); BIdx += AllBoundaryVerts[BIdx] + 1, CurBdry++)
 		{
 			int32 ChainLen = AllBoundaryVerts[BIdx];
-			int32 EndIdx = EndIdx = BIdx + ChainLen + 1;
+			int32 EndIdx = BIdx + ChainLen + 1;
 			check(ChainLen > 2 && EndIdx <= AllBoundaryVerts.Num());
 			if (HoleAreas[CurBdry] < MinHoleArea)
 			{
