minor fixes for fracture overlay splitting
 - Split based on tangent and bitangent
 - Allow vertices w/ almost-identical-but-different attribute values (UV, normal, tangent) to be left un-split

#rnx

diff --git a/Engine/Plugins/Experimental/PlanarCutPlugin/Source/PlanarCut/Private/GeometryMeshConversion.cpp b/Engine/Plugins/Experimental/PlanarCutPlugin/Source/PlanarCut/Private/GeometryMeshConversion.cpp
--- a/Engine/Plugins/Experimental/PlanarCutPlugin/Source/PlanarCut/Private/GeometryMeshConversion.cpp
+++ b/Engine/Plugins/Experimental/PlanarCutPlugin/Source/PlanarCut/Private/GeometryMeshConversion.cpp
@@ -252,7 +252,7 @@ namespace AugmentedDynamicMesh
 		Vs->GetValue(VID, V);
 	}
 
-	template<typename FOverlay>
+	template<typename FOverlay, typename VecType>
 	void SplitOverlayHelper(UE::Geometry::FDynamicMesh3& Mesh, FOverlay* Overlay)
 	{
 		TArray<int32> ContigTris, ContigGroupLens;
@@ -282,6 +282,7 @@ namespace AugmentedDynamicMesh
 				checkSlow(TriStartIdx < ContigTris.Num());
 				int32 TID0 = ContigTris[TriStartIdx];
 				int32 EID0 = GetEID(VID, TID0);
+				VecType LastElData = Overlay->GetElement(EID0);
 				int32 LastEID = EID0;
 				bool bPastEID0 = false;
 				ToSplitTris.Reset();
@@ -289,7 +290,9 @@ namespace AugmentedDynamicMesh
 				{
 					int32 TID = ContigTris[Idx];
 					int32 EID = GetEID(VID, TID);
-					if (EID != LastEID)
+					VecType ElData;
+					if (EID != LastEID &&
+						!VectorUtil::EpsilonEqual(ElData = Overlay->GetElement(EID), LastElData, FMathf::ZeroTolerance))
 					{
 						if (ToSplitTris.Num() > 0)
 						{
@@ -298,6 +301,7 @@ namespace AugmentedDynamicMesh
 							ToSplitTris.Reset();
 						}
 						LastEID = EID;
+						LastElData = ElData;
 						bPastEID0 = true;
 					}
 					if (bPastEID0)
@@ -329,7 +333,7 @@ namespace AugmentedDynamicMesh
 			for (int32 LayerIdx = 0; LayerIdx < NumUVLayers; LayerIdx++)
 			{
 				FDynamicMeshUVOverlay* UVOverlay = Mesh.Attributes()->GetUVLayer(LayerIdx);
-				SplitOverlayHelper<FDynamicMeshUVOverlay>(Mesh, UVOverlay);
+				SplitOverlayHelper<FDynamicMeshUVOverlay, FVector2f>(Mesh, UVOverlay);
 			}
 		}
 
@@ -337,7 +341,7 @@ namespace AugmentedDynamicMesh
 		{
 			for (int32 LayerIdx = 0; LayerIdx < Mesh.Attributes()->NumNormalLayers(); LayerIdx++)
 			{
-				SplitOverlayHelper<FDynamicMeshNormalOverlay>(Mesh, Mesh.Attributes()->GetNormalLayer(0));
+				SplitOverlayHelper<FDynamicMeshNormalOverlay, FVector3f>(Mesh, Mesh.Attributes()->GetNormalLayer(LayerIdx));
 			}
 		}
 
