fix crash on hole-filling fracture trying to add non-manifold triangle

#rnx

diff --git a/Engine/Plugins/Experimental/PlanarCutPlugin/Source/PlanarCut/Private/GeometryMeshConversion.cpp b/Engine/Plugins/Experimental/PlanarCutPlugin/Source/PlanarCut/Private/GeometryMeshConversion.cpp
--- a/Engine/Plugins/Experimental/PlanarCutPlugin/Source/PlanarCut/Private/GeometryMeshConversion.cpp
+++ b/Engine/Plugins/Experimental/PlanarCutPlugin/Source/PlanarCut/Private/GeometryMeshConversion.cpp
@@ -947,6 +947,7 @@ namespace AugmentedDynamicMesh
 			GetTangent(Mesh, HoleEdgeV.A, HoleTanU, HoleTanV);
 			HoleColor = Mesh.GetVertexColor(HoleEdgeV.A);
 			
+			int32 PrevMaxVID = Mesh.MaxVertexID();
 			// Create new vertices for all that weren't already assigned to an existing vertex
 			for (int32 Idx = BIdx + 1; Idx < EndIdx; Idx++)
 			{
@@ -982,8 +983,33 @@ namespace AugmentedDynamicMesh
 				P2 = Mesh.GetVertex(V2);
 
 				int32 TID = Mesh.AppendTriangle(V0, V2, V1);
-				MaterialIDs->SetValue(TID, HoleMaterial);
-				SetVisibility(Mesh, TID, true);
+				// append can fail from a non-manifold edge; in that case add new vertices and try again
+				if (TID == FDynamicMesh3::NonManifoldID)
+				{
+					int32 SeparatedVIDs[]{ V0, V2, V1 };
+					for (int32 SubIdx = 0; SubIdx < 3; SubIdx++)
+					{
+						if (SeparatedVIDs[SubIdx] < PrevMaxVID)
+						{
+							int32 OrigV = SeparatedVIDs[SubIdx];
+							int32 NewV = Mesh.AppendVertex(Mesh, OrigV);
+							SetTangent(Mesh, NewV, HoleNormals[CurBdry], HoleTanU, HoleTanV);
+							for (int32 UVIdx = 0; UVIdx < NumUVs; UVIdx++)
+							{
+								FVector2f UV;
+								GetUV(Mesh, OrigV, UV, UVIdx);
+								SetUV(Mesh, NewV, UV, UVIdx);
+							}
+							SeparatedVIDs[SubIdx] = NewV;
+						}
+					}
+					TID = Mesh.AppendTriangle(SeparatedVIDs[0], SeparatedVIDs[1], SeparatedVIDs[2]);
+				}
+				if (TID > -1)
+				{
+					MaterialIDs->SetValue(TID, HoleMaterial);
+					SetVisibility(Mesh, TID, true);
+				}
 			}
 		}
 	}
