Remove calls to dynamic mesh's GetCachedBounds() from fracture algorithm code's GeometryMeshConversion, to prep for that function being removed from FDynamicMesh3

#rnx

diff --git a/Engine/Plugins/Experimental/PlanarCutPlugin/Source/PlanarCut/Private/GeometryMeshConversion.cpp b/Engine/Plugins/Experimental/PlanarCutPlugin/Source/PlanarCut/Private/GeometryMeshConversion.cpp
--- a/Engine/Plugins/Experimental/PlanarCutPlugin/Source/PlanarCut/Private/GeometryMeshConversion.cpp
+++ b/Engine/Plugins/Experimental/PlanarCutPlugin/Source/PlanarCut/Private/GeometryMeshConversion.cpp
@@ -1280,7 +1280,7 @@ void FCellMeshes::ApplyGeneralGrout(double Grout)
 			VertexCentroid += V;
 		}
 		VertexCentroid /= (double)Mesh.VertexCount();
-		FAxisAlignedBox3d Bounds = Mesh.GetCachedBounds();
+		FAxisAlignedBox3d Bounds = Mesh.GetBounds(true);
 		double BoundsSize = Bounds.MaxDim();
 		// currently just scale the meshes down so they leave half-a-grout worth of space on their longest axis
 		// or delete the mesh if it's so small that that would require a negative scale
@@ -2074,7 +2074,7 @@ void FDynamicMeshCollection::Init(const FGeometryCollection* Collection, const T
 			Editor.RemoveIsolatedVertices();
 		}
 
-		Bounds.Contain(Mesh.GetCachedBounds());
+		Bounds.Contain(MeshData.GetCachedBounds());
 
 		// TODO: build spatial data (add this after setting up mesh boolean path that can use it)
 		//MeshData.Spatial.SetMesh(&Mesh);
@@ -2190,7 +2190,7 @@ int32 FDynamicMeshCollection::CutWithMultiplePlanes(
 		{
 			return false;
 		}
-		if (!ToCut[A]->AugMesh.GetCachedBounds().Intersects(ToCut[B]->AugMesh.GetCachedBounds()))
+		if (!ToCut[A]->GetCachedBounds().Intersects(ToCut[B]->GetCachedBounds()))
 		{
 			return false;
 		}
@@ -2244,7 +2244,7 @@ int32 FDynamicMeshCollection::CutWithMultiplePlanes(
 			int32 TransformIndex = Surface.TransformIndex;
 			FTransform ToCollection = Surface.ToCollection;
 
-			FAxisAlignedBox3d Box = Surface.AugMesh.GetCachedBounds();
+			FAxisAlignedBox3d Box = Surface.GetCachedBounds();
 			if (InternalSurfaceMaterials.NoiseSettings.IsSet())
 			{
 				Box.Expand(InternalSurfaceMaterials.NoiseSettings->Amplitude);
@@ -2304,7 +2304,7 @@ int32 FDynamicMeshCollection::CutWithMultiplePlanes(
 				{
 					if (SplitIslands(ToCut[ResultIndices[UnsplitIdx]]->AugMesh, SplitMeshes))
 					{
-						ToCut[ResultIndices[UnsplitIdx]]->AugMesh = SplitMeshes[0];
+						ToCut[ResultIndices[UnsplitIdx]]->SetMesh(SplitMeshes[0]);
 						for (int32 Idx = 1; Idx < SplitMeshes.Num(); Idx++)
 						{
 							const FDynamicMesh3& Mesh = SplitMeshes[Idx];
@@ -2450,14 +2450,21 @@ int32 FDynamicMeshCollection::CutWithCellMeshes(const FInternalSurfaceMaterials&
 	int BadCount = 0;
 	bool bHasProximity = Collection->HasAttribute("Proximity", FGeometryCollection::GeometryGroup);
 	TArray<int32> GeometryForRemoval;
+	TArray<FAxisAlignedBox3d> CellBounds; CellBounds.Reserve(CellMeshes.CellMeshes.Num());
+	for (int CellIdx = 0; CellIdx < CellMeshes.CellMeshes.Num(); CellIdx++)
+	{
+		CellBounds.Add(CellMeshes.CellMeshes[CellIdx].AugMesh.GetBounds(true));
+	}
 	for (FMeshData& Surface : Meshes)
 	{
 		int32 GeometryIdx = Collection->TransformToGeometryIndex[Surface.TransformIndex];
 		TArray<TUniquePtr<FDynamicMesh3>> BooleanResults;  BooleanResults.SetNum(CellMeshes.CellMeshes.Num());
-		ParallelFor(CellMeshes.CellMeshes.Num(), [&BooleanResults, &CellMeshes, &Surface](int32 CellIdx)
+		FAxisAlignedBox3d SurfaceBounds = Surface.GetCachedBounds();
+
+		ParallelFor(CellMeshes.CellMeshes.Num(), [&BooleanResults, &CellMeshes, &CellBounds, &Surface, &SurfaceBounds](int32 CellIdx)
 			{
 				FCellMeshes::FCellInfo& Cell = CellMeshes.CellMeshes[CellIdx];
-				if (Cell.AugMesh.GetCachedBounds().Intersects(Surface.AugMesh.GetCachedBounds()))
+				if (CellBounds[CellIdx].Intersects(SurfaceBounds))
 				{
 					BooleanResults[CellIdx] = MakeUnique<FDynamicMesh3>();
 					FDynamicMesh3& AugBoolResult = *BooleanResults[CellIdx];
@@ -2556,16 +2563,19 @@ int32 FDynamicMeshCollection::CutWithCellMeshes(const FInternalSurfaceMaterials&
 			{
 				TManagedArray<TSet<int32>>& Proximity = Collection->GetAttribute<TSet<int32>>("Proximity", FGeometryCollection::GeometryGroup);
 				TArray<TUniquePtr<TPointHashGrid3d<int>>> VertexHashes;
-				auto MakeHash = [this, &VertexHashes, &BooleanResults](int GID)
+				TArray<FAxisAlignedBox3d> MeshBounds;
+				auto MakeHash = [this, &VertexHashes, &MeshBounds, &BooleanResults](int GID)
 				{
 					if (GID >= VertexHashes.Num())
 					{
 						VertexHashes.SetNum(GID + 1);
+						MeshBounds.SetNum(GID + 1);
 					}
 					if (!VertexHashes[GID].IsValid())
 					{
 						VertexHashes[GID] = MakeUnique<TPointHashGrid3d<int>>(FMathd::ZeroTolerance * 1000, -1);
 						FillVertexHash(*BooleanResults[GID], *VertexHashes[GID]);
+						MeshBounds[GID] = BooleanResults[GID]->GetBounds(true);
 					}
 				};
 				for (int32 PlaneIdx : PlanesInOutput)
@@ -2592,7 +2602,7 @@ int32 FDynamicMeshCollection::CutWithCellMeshes(const FInternalSurfaceMaterials&
 								{
 									int MeshB = GeometryToResultMesh[GIDB];
 									MakeHash(MeshB);
-									if (IsNeighboring(*BooleanResults[MeshA], *VertexHashes[MeshA], *BooleanResults[MeshB], *VertexHashes[MeshB]))
+									if (IsNeighboring(*BooleanResults[MeshA], *VertexHashes[MeshA], MeshBounds[MeshA], *BooleanResults[MeshB], *VertexHashes[MeshB], MeshBounds[MeshB]))
 									{
 										Proximity[GIDA].Add(GIDB);
 										Proximity[GIDB].Add(GIDA);
@@ -2629,14 +2639,18 @@ void FDynamicMeshCollection::FillVertexHash(const FDynamicMesh3& Mesh, TPointHas
 	}
 }
 
-
-bool FDynamicMeshCollection::IsNeighboring(FDynamicMesh3* Mesh[2], const TPointHashGrid3d<int>* VertHash[2])
+bool FDynamicMeshCollection::IsNeighboring(
+	UE::Geometry::FDynamicMesh3& MeshA, const UE::Geometry::TPointHashGrid3d<int>& VertHashA, const UE::Geometry::FAxisAlignedBox3d& BoundsA,
+	UE::Geometry::FDynamicMesh3& MeshB, const UE::Geometry::TPointHashGrid3d<int>& VertHashB, const UE::Geometry::FAxisAlignedBox3d& BoundsB)
 {
+	UE::Geometry::FDynamicMesh3* Mesh[2]{ &MeshA, &MeshB };
+	const UE::Geometry::TPointHashGrid3d<int>* VertHash[2]{ &VertHashA, &VertHashB };
+
 	if (!ensure(Mesh[0] && Mesh[1] && VertHash[0] && VertHash[1]))
 	{
 		return false;
 	}
-	if (!Mesh[0]->GetCachedBounds().Intersects(Mesh[1]->GetCachedBounds()))
+	if (!BoundsA.Intersects(BoundsB))
 	{
 		return false;
 	}
diff --git a/Engine/Plugins/Experimental/PlanarCutPlugin/Source/PlanarCut/Private/GeometryMeshConversion.h b/Engine/Plugins/Experimental/PlanarCutPlugin/Source/PlanarCut/Private/GeometryMeshConversion.h
--- a/Engine/Plugins/Experimental/PlanarCutPlugin/Source/PlanarCut/Private/GeometryMeshConversion.h
+++ b/Engine/Plugins/Experimental/PlanarCutPlugin/Source/PlanarCut/Private/GeometryMeshConversion.h
@@ -166,6 +166,32 @@ struct PLANARCUT_API FDynamicMeshCollection
 
 		FMeshData(const UE::Geometry::FDynamicMesh3& Mesh, int32 TransformIndex, FTransform ToCollection) : AugMesh(Mesh), TransformIndex(TransformIndex), ToCollection(ToCollection)
 		{}
+
+		void SetMesh(const UE::Geometry::FDynamicMesh3& NewAugMesh)
+		{
+			ClearCachedBounds();
+			AugMesh = NewAugMesh;
+		}
+
+		/// Note this relies on the caller to also call ClearCachedBounds() as needed; it will not automatically invalidate any computed bounds
+		const UE::Geometry::FAxisAlignedBox3d& GetCachedBounds()
+		{
+			if (!bHasBounds)
+			{
+				Bounds = AugMesh.GetBounds(true);
+				bHasBounds = true;
+			}
+			return Bounds;
+		}
+
+		void ClearCachedBounds()
+		{
+			bHasBounds = false;
+		}
+
+	private:
+		bool bHasBounds = false;
+		UE::Geometry::FAxisAlignedBox3d Bounds;
 	};
 	TIndirectArray<FMeshData> Meshes;
 	UE::Geometry::FAxisAlignedBox3d Bounds;
@@ -230,14 +256,9 @@ private:
 
 	void FillVertexHash(const UE::Geometry::FDynamicMesh3& Mesh, UE::Geometry::TPointHashGrid3d<int>& VertHash);
 
-	bool IsNeighboring(UE::Geometry::FDynamicMesh3& MeshA, const UE::Geometry::TPointHashGrid3d<int>& VertHashA, UE::Geometry::FDynamicMesh3& MeshB, const UE::Geometry::TPointHashGrid3d<int>& VertHashB)
-	{
-		UE::Geometry::FDynamicMesh3* Mesh[2]{ &MeshA, &MeshB };
-		const UE::Geometry::TPointHashGrid3d<int>* VertHash[2]{ &VertHashA, &VertHashB };
-		return IsNeighboring(Mesh, VertHash);
-	}
-
-	bool IsNeighboring(UE::Geometry::FDynamicMesh3* Mesh[2], const UE::Geometry::TPointHashGrid3d<int>* VertHash[2]);
+	bool IsNeighboring(
+		UE::Geometry::FDynamicMesh3& MeshA, const UE::Geometry::TPointHashGrid3d<int>& VertHashA, const UE::Geometry::FAxisAlignedBox3d& BoundsA,
+		UE::Geometry::FDynamicMesh3& MeshB, const UE::Geometry::TPointHashGrid3d<int>& VertHashB, const UE::Geometry::FAxisAlignedBox3d& BoundsB);
 };
 
 
