fixes for fracture hole filling and related issues in mesh boolean code

#rnx

diff --git a/Engine/Plugins/Experimental/PlanarCutPlugin/Source/PlanarCut/Private/GeometryMeshConversion.cpp b/Engine/Plugins/Experimental/PlanarCutPlugin/Source/PlanarCut/Private/GeometryMeshConversion.cpp
--- a/Engine/Plugins/Experimental/PlanarCutPlugin/Source/PlanarCut/Private/GeometryMeshConversion.cpp
+++ b/Engine/Plugins/Experimental/PlanarCutPlugin/Source/PlanarCut/Private/GeometryMeshConversion.cpp
@@ -616,11 +616,11 @@ namespace AugmentedDynamicMesh
 	 * @param Mesh Fill holes on this mesh
 	 * @param CandidateEdges Only consider holes that touch these edges
 	 */
-	void FillHoles(FDynamicMesh3& Mesh, const TSet<int32>& CandidateEdges, double MinHoleArea = KINDA_SMALL_NUMBER)
+	void FillHoles(FDynamicMesh3& Mesh, const TSet<int32>& CandidateEdges, double MinHoleArea = DOUBLE_KINDA_SMALL_NUMBER)
 	{
 		/// 1. Build a spatial hash of boundary vertices, so we can identify holes even though the mesh is not welded
 		
-		double SnapDistance = 1e-05;
+		double SnapDistance = 1e-03;
 		TPointHashGrid3d<int32> VertHash(SnapDistance * 10, -1);
 		TMap<int32, int32> CoincidentVerticesMap; // map every vertex in an overlapping cluster to a canonical single vertex for that cluster
 		TSet<int32> HashedVertices; // track what's already processed
@@ -782,8 +782,6 @@ namespace AugmentedDynamicMesh
 
 		TArray<int32> HoleComponents; // which connected component to connect the new hole geometry to
 		HoleComponents.Init(-1, NumHoles);
-		TArray<double> HoleAreas;
-		HoleAreas.Init(0, NumHoles);
 		TArray<FVector3d> HoleNormals;
 		HoleNormals.Init(FVector3d::ZeroVector, NumHoles);
 		for (int32 BIdx = 0, CurBdry = 0; BIdx < AllBoundaryVerts.Num(); BIdx += AllBoundaryVerts[BIdx] + 1, CurBdry++)
@@ -793,19 +791,21 @@ namespace AugmentedDynamicMesh
 			check(ChainLen > 2 && EndIdx <= AllBoundaryVerts.Num());
 
 			TMap<int32, FVector3d> ComponentToNormalMap;
-			FVector3d HoleNormal(0, 0, 0), LastEdge(0, 0, 0);
-			double Area = 0;
+			FVector3d HoleNormal(0, 0, 0);
+			FVector3d POrigin = Mesh.GetVertex(AllBoundaryVerts[BIdx + 1]);
+			
 			for (int32 LastIdx = EndIdx - 1, Idx = BIdx + 1; Idx < EndIdx; LastIdx = Idx++)
 			{
-				int32 V0 = AllBoundaryVerts[LastIdx];
-				int32 V1 = AllBoundaryVerts[Idx];
-				FVector3d P0 = Mesh.GetVertex(V0);
-				FVector3d P1 = Mesh.GetVertex(V1);
-				FVector3d Edge = P1 - P0;
-				FVector3d HoleTriNormal = Edge.Cross(LastEdge);
-				Area += HoleTriNormal.Length() * .5;
-				HoleNormal += HoleTriNormal;
-				LastEdge = Edge;
+				if (LastIdx != BIdx + 1 && Idx != BIdx + 1)
+				{
+					int32 V0 = AllBoundaryVerts[LastIdx];
+					int32 V1 = AllBoundaryVerts[Idx];
+					FVector3d P0 = Mesh.GetVertex(V0);
+					FVector3d P1 = Mesh.GetVertex(V1);
+					FVector3d Edge = P1 - P0;
+					FVector3d HoleTriNormal = Edge.Cross(P0 - POrigin);
+					HoleNormal += HoleTriNormal;
+				}
 
 				int32 EID01 = AllBoundaryEdges[LastIdx];
 				FDynamicMesh3::FEdge Edge01 = Mesh.GetEdge(EID01);
@@ -848,16 +848,18 @@ namespace AugmentedDynamicMesh
 				}
 			}
 
-			HoleAreas[CurBdry] = Area;
 			if (!bHasNormal && BestComponent > -1)
 			{
 				HoleNormal = ComponentToNormalMap[BestComponent];
 			}
-			if (BestComponent == -1)
-			{
-				ensure(false);
-				BestComponent = VertComponents.Find(AllBoundaryVerts[BIdx + 1]);
-			}
+			// Note: BestComponent could be -1; currently we ignore the hole in this case.
+			// This is the case where only the "outside" material is connected to the hole.
+			// it tends to only happen for e.g. singular degenerate triangles, which are 
+			// better to ignore.
+			// If we'd instead like to fill these anyway, we could fall back to:
+			//   BestComponent = VertComponents.Find(AllBoundaryVerts[BIdx + 1])
+			// Or we could change the hole fill logic to create a completely
+			// disconnected patch in this case.
 			HoleComponents[CurBdry] = BestComponent;
 			HoleNormals[CurBdry] = HoleNormal;
 		}
@@ -869,10 +871,44 @@ namespace AugmentedDynamicMesh
 		TArray<int32> UseVIDs; // vertex IDs for the current hole
 		for (int32 BIdx = 0, CurBdry = 0; BIdx < AllBoundaryVerts.Num(); BIdx += AllBoundaryVerts[BIdx] + 1, CurBdry++)
 		{
+			if (HoleComponents[CurBdry] == -1)
+			{
+				// hole was solely on the "outside" material; this can indicate an open boundary that isn't actually a hole
+				// for now we skip these cases (See comment where HoleComponents are assigned, above)
+				continue;
+			}
+
 			int32 ChainLen = AllBoundaryVerts[BIdx];
 			int32 EndIdx = BIdx + ChainLen + 1;
 			check(ChainLen > 2 && EndIdx <= AllBoundaryVerts.Num());
-			if (HoleAreas[CurBdry] < MinHoleArea)
+
+			TArray<FIndex3i> Triangles;
+			TArray<UE::Geometry::FVector3<double>> VertexPositions;
+			for (int Idx = BIdx + 1; Idx < EndIdx; Idx++)
+			{
+				VertexPositions.Add(Mesh.GetVertex(AllBoundaryVerts[Idx]));
+			}
+			PolygonTriangulation::TriangulateSimplePolygon(VertexPositions, Triangles);
+			double HoleArea = 0;
+			FVector3d LastNormal(0, 0, 0);
+			const double FlippedThreshold = -.5; // Allow some angle deviation but reject hole fills that are too folded
+			bool bHoleTriangulationIsFolded = false;
+			for (FIndex3i Tri : Triangles)
+			{
+				double TriArea;
+				FVector3d TriNormal = VectorUtil::NormalArea(VertexPositions[Tri.A], VertexPositions[Tri.B], VertexPositions[Tri.C], TriArea);
+				HoleArea += TriArea;
+				if (LastNormal.Dot(TriNormal) < FlippedThreshold)
+				{
+					bHoleTriangulationIsFolded = true;
+					break;
+				}
+				if (TriArea != 0.0)
+				{
+					LastNormal = TriNormal;
+				}
+			}
+			if (HoleArea < MinHoleArea || bHoleTriangulationIsFolded)
 			{
 				continue;
 			}
@@ -971,16 +1007,14 @@ namespace AugmentedDynamicMesh
 				UseVIDs[C] = NewV;
 			}
 
-			int32 V0 = UseVIDs[0];
-			int32 V1, V2 = UseVIDs[1];
-			FVector3d P0 = Mesh.GetVertex(V0);
-			FVector3d P1, P2 = Mesh.GetVertex(V2);
-			for (int32 C1 = 1, C2 = 2; C2 < ChainLen; C1 = C2++)
+			for (FIndex3i Tri : Triangles)
 			{
-				V1 = V2;
-				V2 = UseVIDs[C2];
-				P1 = P2;
-				P2 = Mesh.GetVertex(V2);
+				int32 V0 = UseVIDs[Tri.A];
+				int32 V1 = UseVIDs[Tri.B];
+				int32 V2 = UseVIDs[Tri.C];
+				FVector3d P0 = Mesh.GetVertex(V0);
+				FVector3d P1 = Mesh.GetVertex(V1);
+				FVector3d P2 = Mesh.GetVertex(V2);
 
 				int32 TID = Mesh.AppendTriangle(V0, V2, V1);
 				// append can fail from a non-manifold edge; in that case add new vertices and try again
@@ -2879,7 +2913,11 @@ int32 FDynamicMeshCollection::AppendToCollection(const FTransform& ToCollection,
 	int32 NewGeometryStartIdx = Output.FaceStart.Num();
 	int32 OriginalVertexNum = Output.Vertex.Num();
 	int32 OriginalFaceNum = Output.Indices.Num();
-	int32 UVLayerCount = Output.NumUVLayers();
+	int32 UVLayerCount = AugmentedDynamicMesh::NumEnabledUVChannels(Mesh);
+	if (!ensure(UVLayerCount > 0))
+	{
+		UVLayerCount = 1;
+	}
 
 	int32 GeometryIdx = Output.AddElements(1, FGeometryCollection::GeometryGroup);
 	int32 TransformIdx = Output.AddElements(1, FGeometryCollection::TransformGroup);
diff --git a/Engine/Plugins/Runtime/GeometryProcessing/Source/DynamicMesh/Private/Operations/MeshBoolean.cpp b/Engine/Plugins/Runtime/GeometryProcessing/Source/DynamicMesh/Private/Operations/MeshBoolean.cpp
--- a/Engine/Plugins/Runtime/GeometryProcessing/Source/DynamicMesh/Private/Operations/MeshBoolean.cpp
+++ b/Engine/Plugins/Runtime/GeometryProcessing/Source/DynamicMesh/Private/Operations/MeshBoolean.cpp
@@ -429,16 +429,15 @@ bool FMeshBoolean::Compute()
 				});
 			for (int EID : ProcessMesh.EdgeIndicesItr())
 			{
-				FIndex2i TriPair = ProcessMesh.GetEdgeT(EID);
-				if (TriPair.B == IndexConstants::InvalidID || KeepTri[MeshIdx][TriPair.A] == KeepTri[MeshIdx][TriPair.B])
+				FDynamicMesh3::FEdge Edge = ProcessMesh.GetEdge(EID);
+				if (Edge.Tri.B == IndexConstants::InvalidID || KeepTri[MeshIdx][Edge.Tri.A] == KeepTri[MeshIdx][Edge.Tri.B])
 				{
 					continue;
 				}
 
 				CutBoundaryEdges[MeshIdx].Add(EID);
-				FIndex2i VertPair = ProcessMesh.GetEdgeV(EID);
-				PossUnmatchedBdryVerts[MeshIdx].Add(VertPair.A);
-				PossUnmatchedBdryVerts[MeshIdx].Add(VertPair.B);
+				PossUnmatchedBdryVerts[MeshIdx].Add(Edge.Vert.A);
+				PossUnmatchedBdryVerts[MeshIdx].Add(Edge.Vert.B);
 			}
 		}
 		// now go ahead and delete from both meshes
@@ -876,7 +875,7 @@ bool FMeshBoolean::CollapseWouldChangeShapeOrUVs(
 					bHasBadEdge = true;
 					return; // break instead of continue to skip the whole edge
 				}
-				if (OtherEdgeDir.Dot(EdgeDir) <= -1 + DotTolerance)
+				if (OtherEdgeDir.Dot(EdgeDir) <= -DotTolerance)
 				{
 					OpposedEdge = VertEID;
 				}
@@ -1151,13 +1150,31 @@ void FMeshBoolean::SimplifyAlongNewEdges(int NumMeshesToProcess, FDynamicMesh3*
 				FDynamicMesh3::FEdgeCollapseInfo CollapseInfo;
 				int RemoveV = Edge.Vert[RemoveVIdx];
 				int KeepV = Edge.Vert[KeepVIdx];
+				// Detect the case of a triangle with two boundary edges, where collapsing 
+				// the target boundary edge would keep the non-boundary edge.
+				// This collapse will remove the triangle, so we add the
+				// (formerly) non-boundary edge as our new boundary edge.
+				auto WouldRemoveTwoBoundaryEdges = [](const FDynamicMesh3& Mesh, int EID, int RemoveV)
+				{
+					checkSlow(Mesh.IsEdge(EID));
+					int OppV = Mesh.GetEdgeOpposingV(EID).A;
+					int NextOnTri = Mesh.FindEdge(RemoveV, OppV);
+					return Mesh.IsBoundaryEdge(NextOnTri);
+				};
+				bool bWouldRemoveNext = WouldRemoveTwoBoundaryEdges(*CutMesh[0], EID, RemoveV);
 				EMeshResult CollapseResult = CutMesh[0]->CollapseEdge(KeepV, RemoveV, 0, CollapseInfo);
 				if (CollapseResult == EMeshResult::Ok)
 				{
+					if (bWouldRemoveNext && ensure(CutMesh[0]->IsBoundaryEdge(CollapseInfo.KeptEdges.A)))
+					{
+						CutBoundaryEdgeSets[0].Add(CollapseInfo.KeptEdges.A);
+					}
+
 					if (bHasMatches)
 					{
 						int OtherRemoveV = Matches[RemoveVIdx];
 						int OtherKeepV = Matches[KeepVIdx];
+						bool bOtherWouldRemoveNext = WouldRemoveTwoBoundaryEdges(*CutMesh[1], OtherEID, OtherRemoveV);
 						FDynamicMesh3::FEdgeCollapseInfo OtherCollapseInfo;
 						EMeshResult OtherCollapseResult = CutMesh[1]->CollapseEdge(OtherKeepV, OtherRemoveV, 0, OtherCollapseInfo);
 						if (OtherCollapseResult != EMeshResult::Ok)
@@ -1170,6 +1187,10 @@ void FMeshBoolean::SimplifyAlongNewEdges(int NumMeshesToProcess, FDynamicMesh3*
 						}
 						else
 						{
+							if (bOtherWouldRemoveNext && ensure(CutMesh[1]->IsBoundaryEdge(OtherCollapseInfo.KeptEdges.A)))
+							{
+								CutBoundaryEdgeSets[1].Add(OtherCollapseInfo.KeptEdges.A);
+							}
 							AllVIDMatches.Remove(RemoveV);
 							CutBoundaryEdgeSets[1].Remove(OtherCollapseInfo.CollapsedEdge);
 							CutBoundaryEdgeSets[1].Remove(OtherCollapseInfo.RemovedEdges[0]);
