fracture tinygeo tool improvements: add option to also merge-away the selected bones + add another option for choosing which neighbor to merge to ("closest center" vs "largest volume")

diff --git a/Engine/Plugins/Experimental/ChaosEditor/Source/FractureEditor/Private/FractureToolFixTinyGeo.cpp b/Engine/Plugins/Experimental/ChaosEditor/Source/FractureEditor/Private/FractureToolFixTinyGeo.cpp
--- a/Engine/Plugins/Experimental/ChaosEditor/Source/FractureEditor/Private/FractureToolFixTinyGeo.cpp
+++ b/Engine/Plugins/Experimental/ChaosEditor/Source/FractureEditor/Private/FractureToolFixTinyGeo.cpp
@@ -109,6 +109,16 @@ void UFractureToolFixTinyGeo::FractureContextChanged()
 			SmallIndices
 		);
 
+		if (TinyGeoSettings->bAlsoMergeSelected)
+		{
+			const TArray<int32> Selection = FractureContext.GetSelection();
+			for (int32 Bone : Selection)
+			{
+				// note: we AddUnique when doing the merge, but for rendering we allow duplicates
+				SmallIndices.Add(Bone);
+			}
+		}
+
 		FTransform OuterTransform = FractureContext.GetTransform();
 		for (int32 TransformIdx : SmallIndices) // small transforms
 		{
@@ -183,13 +193,31 @@ int32 UFractureToolFixTinyGeo::ExecuteFracture(const FFractureToolContext& Fract
 			MinVolume,
 			SmallIndices
 		);
+
+		if (TinyGeoSettings->bAlsoMergeSelected)
+		{
+			const TArray<int32> Selection = FractureContext.GetSelection();
+			for (int32 Bone : Selection)
+			{
+				SmallIndices.AddUnique(Bone);
+			}
+		}
+
+		// convert ENeighborSelectionMethod to the PlanarCut module version (the only difference is it's not a UENUM)
+		UE::PlanarCut::ENeighborSelectionMethod SelectionMethod = UE::PlanarCut::ENeighborSelectionMethod::LargestNeighbor;
+		if (TinyGeoSettings->NeighborSelection == ENeighborSelectionMethod::NearestCenter)
+		{
+			SelectionMethod = UE::PlanarCut::ENeighborSelectionMethod::NearestCenter;
+		}
+
 		MergeBones(
 			Collection,
 			TransformIndices,
 			Volumes,
 			MinVolume,
 			SmallIndices,
-			false /*bUnionJoinedPieces*/ // Note: Union-ing the pieces is nicer in theory, but can leave cracks and non-manifold garbage
+			false, /*bUnionJoinedPieces*/ // Note: Union-ing the pieces is nicer in theory, but can leave cracks and non-manifold garbage
+			SelectionMethod
 		);
 		UpdateIdx = 0;
 	}
diff --git a/Engine/Plugins/Experimental/ChaosEditor/Source/FractureEditor/Private/FractureToolFixTinyGeo.h b/Engine/Plugins/Experimental/ChaosEditor/Source/FractureEditor/Private/FractureToolFixTinyGeo.h
--- a/Engine/Plugins/Experimental/ChaosEditor/Source/FractureEditor/Private/FractureToolFixTinyGeo.h
+++ b/Engine/Plugins/Experimental/ChaosEditor/Source/FractureEditor/Private/FractureToolFixTinyGeo.h
@@ -16,10 +16,18 @@ enum class EGeometrySelectionMethod
 	VolumeCubeRoot,
 	// Select by cube root of volume relative to the overall shape's cube root of volume
 	RelativeVolume
-	//~ TODO: also have a manual selection method?
 };
 
-/** Settings controlling how convex hulls are generated for geometry collections */
+UENUM()
+enum class ENeighborSelectionMethod
+{
+	// Merge to the neighbor with the largest volume
+	LargestNeighbor,
+	// Merge to the neighbor with the closest center
+	NearestCenter
+};
+
+/** Settings controlling how tiny geometry is selected and merged into neighboring geometry */
 UCLASS(config = EditorPerProjectUserSettings)
 class UFractureTinyGeoSettings : public UFractureToolSettings
 {
@@ -34,15 +42,22 @@ public:
 	UPROPERTY(EditAnywhere, Category = FilterSettings)
 	EGeometrySelectionMethod SelectionMethod = EGeometrySelectionMethod::RelativeVolume;
 
+	UPROPERTY(EditAnywhere, Category = FilterSettings)
+	ENeighborSelectionMethod NeighborSelection = ENeighborSelectionMethod::LargestNeighbor;
+
 	/** If volume is less than this value cubed, geometry should be merged into neighbors -- i.e. a value of 2 merges geometry smaller than a 2x2x2 cube */
 	UPROPERTY(EditAnywhere, Category = FilterSettings, meta = (ClampMin = ".00001", UIMin = ".1", UIMax = "10", EditCondition = "SelectionMethod == EGeometrySelectionMethod::VolumeCubeRoot", EditConditionHides))
 	double MinVolumeCubeRoot = 1;
 
 	/** If cube root of volume relative to the overall shape's cube root of volume is less than this, the geometry should be merged into its neighbors.
 	(Note: This is a bit different from the histogram viewer's "Relative Size," which instead shows values relative to the largest rigid bone.) */
-	UPROPERTY(EditAnywhere, Category = FilterSettings, meta = (ClampMin = ".001", UIMax = ".1", ClampMax = "1.0", EditCondition = "SelectionMethod == EGeometrySelectionMethod::RelativeVolume", EditConditionHides))
+	UPROPERTY(EditAnywhere, Category = FilterSettings, meta = (ClampMin = "0", UIMax = ".1", ClampMax = "1.0", EditCondition = "SelectionMethod == EGeometrySelectionMethod::RelativeVolume", EditConditionHides))
 	double RelativeVolume = .01;
 
+	/** Also merge selected bones to their neighbors */
+	UPROPERTY(EditAnywhere, Category = FilterSettings)
+	bool bAlsoMergeSelected = false;
+
 };
 
 
diff --git a/Engine/Plugins/Experimental/PlanarCutPlugin/Source/PlanarCut/Private/PlanarCut.cpp b/Engine/Plugins/Experimental/PlanarCutPlugin/Source/PlanarCut/Private/PlanarCut.cpp
--- a/Engine/Plugins/Experimental/PlanarCutPlugin/Source/PlanarCut/Private/PlanarCut.cpp
+++ b/Engine/Plugins/Experimental/PlanarCutPlugin/Source/PlanarCut/Private/PlanarCut.cpp
@@ -884,7 +884,8 @@ int32 MergeBones(
 	const TArrayView<const double>& Volumes,
 	double MinVolume,
 	const TArrayView<const int32>& SmallTransformIndices,
-	bool bUnionJoinedPieces
+	bool bUnionJoinedPieces,
+	UE::PlanarCut::ENeighborSelectionMethod NeighborSelectionMethod
 )
 {
 	using UE::Geometry::FTransform3d;
@@ -1005,6 +1006,23 @@ int32 MergeBones(
 	TSet<int32> TooSmalls;
 	TSet<int32> CanMerge;
 
+	// GeomToCenter is just a cache for GetCenter; may not be worth caching as long as 'center' == bounding box center
+	// TODO: consider switching to a more accurate 'center' and/or removing the cache
+	TMap<int32, FVector> GeomToCenter;
+	auto GetCenter = [&Collection, &GeomToCenter](int32 GeomIdx) -> FVector
+	{
+		FVector* CachedCenter = GeomToCenter.Find(GeomIdx);
+		if (CachedCenter)
+		{
+			return *CachedCenter;
+		}
+		int32 TransformIdx = Collection.TransformIndex[GeomIdx];
+		FTransform Transform = GeometryCollectionAlgo::GlobalMatrix(Collection.Transform, Collection.Parent, TransformIdx);
+		FVector Center = Transform.TransformPosition(Collection.BoundingBox[GeomIdx].GetCenter());
+		GeomToCenter.Add(GeomIdx, Center);
+		return Center;
+	};
+
 	for (int32 TransformIdx : TransformIndices)
 	{
 		int32 GeomIdx = Collection.TransformToGeometryIndex[TransformIdx];
@@ -1040,29 +1058,38 @@ int32 MergeBones(
 		}
 
 		const TSet<int32>& Prox = Proximity[SmallIdx];
-		double BiggestVol = 0;
-		int32 BiggestVolIdx = INDEX_NONE;
+		double BestScore = -FMathd::MaxReal;
+		int32 BestNbrIdx = INDEX_NONE;
 		for (int32 NbrIdx : Prox)
 		{
-			if (CanMerge.Contains(NbrIdx))
+			if (NbrIdx != SmallIdx && CanMerge.Contains(NbrIdx))
 			{
-				double Vol = GeomToVol[NbrIdx];
-				if (Vol > BiggestVol)
+				double Score;
+				if (NeighborSelectionMethod == UE::PlanarCut::ENeighborSelectionMethod::LargestNeighbor)
+				{
+					Score = GeomToVol[NbrIdx];
+				}
+				else // Nearest center
+				{
+					Score = 1.0 / (DOUBLE_SMALL_NUMBER + DistanceSquared(GetCenter(NbrIdx), GetCenter(SmallIdx)));
+				}
+				if (Score > BestScore)
 				{
-					BiggestVol = Vol;
-					BiggestVolIdx = NbrIdx;
+					BestScore = Score;
+					BestNbrIdx = NbrIdx;
 				}
 			}
 		}
-		if (BiggestVolIdx == INDEX_NONE)
+		if (BestNbrIdx == INDEX_NONE)
 		{
+			UE_LOG(LogPlanarCut, Warning, TEXT("Couldn't fix Bone %d: No neighbors found in proximity graph"), SmallIdx);
 			continue;
 		}
 
 		if (SmallRemoveGroupIdx)
 		{
 			int32 OldSGIdx = *SmallRemoveGroupIdx;
-			int32* BigRemoveGroupIdx = GeomIdxToRemoveGroupIdx.Find(BiggestVolIdx);
+			int32* BigRemoveGroupIdx = GeomIdxToRemoveGroupIdx.Find(BestNbrIdx);
 			if (BigRemoveGroupIdx)
 			{
 				int32 BigRGIdx = *BigRemoveGroupIdx;
@@ -1074,14 +1101,14 @@ int32 MergeBones(
 			}
 			else
 			{
-				RemoveGroups[OldSGIdx].AddBig(GeomToVol, MinVolume, BiggestVolIdx);
-				checkSlow(!GeomIdxToRemoveGroupIdx.Contains(BiggestVolIdx));
-				GeomIdxToRemoveGroupIdx.Add(BiggestVolIdx, OldSGIdx);
+				RemoveGroups[OldSGIdx].AddBig(GeomToVol, MinVolume, BestNbrIdx);
+				checkSlow(!GeomIdxToRemoveGroupIdx.Contains(BestNbrIdx));
+				GeomIdxToRemoveGroupIdx.Add(BestNbrIdx, OldSGIdx);
 			}
 		}
 		else
 		{
-			int32* BigRemoveGroupIdx = GeomIdxToRemoveGroupIdx.Find(BiggestVolIdx);
+			int32* BigRemoveGroupIdx = GeomIdxToRemoveGroupIdx.Find(BestNbrIdx);
 			if (BigRemoveGroupIdx)
 			{
 				int32 BigRGIdx = *BigRemoveGroupIdx;
@@ -1091,11 +1118,11 @@ int32 MergeBones(
 			}
 			else
 			{
-				int32 RemoveGroupIdx = RemoveGroups.Emplace(GeomToVol, MinVolume, SmallIdx, BiggestVolIdx);
+				int32 RemoveGroupIdx = RemoveGroups.Emplace(GeomToVol, MinVolume, SmallIdx, BestNbrIdx);
 				checkSlow(!GeomIdxToRemoveGroupIdx.Contains(SmallIdx));
-				checkSlow(!GeomIdxToRemoveGroupIdx.Contains(BiggestVolIdx));
+				checkSlow(!GeomIdxToRemoveGroupIdx.Contains(BestNbrIdx));
 				GeomIdxToRemoveGroupIdx.Add(SmallIdx, RemoveGroupIdx);
-				GeomIdxToRemoveGroupIdx.Add(BiggestVolIdx, RemoveGroupIdx);
+				GeomIdxToRemoveGroupIdx.Add(BestNbrIdx, RemoveGroupIdx);
 			}
 		}
 	}
diff --git a/Engine/Plugins/Experimental/PlanarCutPlugin/Source/PlanarCut/Public/PlanarCut.h b/Engine/Plugins/Experimental/PlanarCutPlugin/Source/PlanarCut/Public/PlanarCut.h
--- a/Engine/Plugins/Experimental/PlanarCutPlugin/Source/PlanarCut/Public/PlanarCut.h
+++ b/Engine/Plugins/Experimental/PlanarCutPlugin/Source/PlanarCut/Public/PlanarCut.h
@@ -243,15 +243,29 @@ void PLANARCUT_API FindSmallBones(
 	TArray<int32>& OutSmallBones
 );
 
+namespace UE
+{
+	namespace PlanarCut
+	{
+		enum ENeighborSelectionMethod
+		{
+			LargestNeighbor,
+			NearestCenter
+		};
+	}
+}
+
+
 /**
- * Merge chosen geometry into neighboring geometry.  Geometry will be merged into its largest neighboring geometry (if any).
+ * Merge chosen geometry into neighboring geometry.
  *
- * @param Collection			The collection to be processed
- * @param TransformIndices		The transform indices to process, or empty if all should be processed
- * @param Volumes				Volumes of geometry; 1:1 w/ TransformIndices array
- * @param MinVolume				If merged small geometry is larger than this, it will not require further merging
- * @param SmallTransformIndices	Transformed indices of pieces that we want to merge
- * @param bUnionJoinedPieces	Try to 'union' the merged pieces, removing internal triangles and connecting the shared cut boundary
+ * @param Collection				The collection to be processed
+ * @param TransformIndices			The transform indices to process, or empty if all should be processed
+ * @param Volumes					Volumes of geometry; 1:1 w/ TransformIndices array
+ * @param MinVolume					If merged small geometry is larger than this, it will not require further merging
+ * @param SmallTransformIndices		Transform indices of pieces that we want to merge
+ * @param bUnionJoinedPieces		Try to 'union' the merged pieces, removing internal triangles and connecting the shared cut boundary
+ * @param NeighborSelectionMethod	How to choose which neighbor to merge to
  */
 int32 PLANARCUT_API MergeBones(
 	FGeometryCollection& Collection,
@@ -259,7 +273,8 @@ int32 PLANARCUT_API MergeBones(
 	const TArrayView<const double>& Volumes,
 	double MinVolume,
 	const TArrayView<const int32>& SmallTransformIndices,
-	bool bUnionJoinedPieces
+	bool bUnionJoinedPieces,
+	UE::PlanarCut::ENeighborSelectionMethod NeighborSelectionMethod
 );
 
 /**
