fracture performance improvement -- remove expensive validation on mergebones for non-debug builds

#rnx

diff --git a/Engine/Plugins/Experimental/PlanarCutPlugin/Source/PlanarCut/Private/PlanarCut.cpp b/Engine/Plugins/Experimental/PlanarCutPlugin/Source/PlanarCut/Private/PlanarCut.cpp
--- a/Engine/Plugins/Experimental/PlanarCutPlugin/Source/PlanarCut/Private/PlanarCut.cpp
+++ b/Engine/Plugins/Experimental/PlanarCutPlugin/Source/PlanarCut/Private/PlanarCut.cpp
@@ -1228,7 +1228,11 @@ int32 MergeBones(
 	}
 
 	// remove transforms for all geometry that was merged in
-	Collection.RemoveElements(FGeometryCollection::TransformGroup, AllRemoveIndices);
+	FManagedArrayCollection::FProcessingParameters ProcessingParams;
+#if !UE_BUILD_DEBUG
+	ProcessingParams.bDoValidation = false;
+#endif
+	Collection.RemoveElements(FGeometryCollection::TransformGroup, AllRemoveIndices, ProcessingParams);
 
 	return INDEX_NONE; // TODO: consider tracking smallest index of updated groups?  but no reason to do so currently
 }
