fix incorrectly transformed geometry when fracturing a geometry collection that has non-uniform scale + rotation, by not using FTransform's Inverse() function

#rnx

diff --git a/Engine/Plugins/Experimental/PlanarCutPlugin/Source/PlanarCut/Private/GeometryMeshConversion.cpp b/Engine/Plugins/Experimental/PlanarCutPlugin/Source/PlanarCut/Private/GeometryMeshConversion.cpp
--- a/Engine/Plugins/Experimental/PlanarCutPlugin/Source/PlanarCut/Private/GeometryMeshConversion.cpp
+++ b/Engine/Plugins/Experimental/PlanarCutPlugin/Source/PlanarCut/Private/GeometryMeshConversion.cpp
@@ -2036,7 +2036,7 @@ void FDynamicMeshCollection::Init(const FGeometryCollection* Collection, const T
 		int32 AddedMeshIdx = Meshes.Add(new FMeshData(NumUVLayers));
 		FMeshData& MeshData = Meshes[AddedMeshIdx];
 		MeshData.TransformIndex = TransformIdx;
-		MeshData.ToCollection = FTransform(CollectionToLocal.Inverse());
+		MeshData.FromCollection = FTransform(CollectionToLocal);
 		FDynamicMesh3& Mesh = MeshData.AugMesh;
 		
 		Mesh.EnableAttributes();
@@ -2278,7 +2278,7 @@ int32 FDynamicMeshCollection::CutWithMultiplePlanes(
 		{
 			FMeshData& Surface = *ToCut[ToCutIdx];
 			int32 TransformIndex = Surface.TransformIndex;
-			FTransform ToCollection = Surface.ToCollection;
+			FTransform FromCollection = Surface.FromCollection;
 
 			FAxisAlignedBox3d Box = Surface.GetCachedBounds();
 			if (InternalSurfaceMaterials.NoiseSettings.IsSet())
@@ -2295,7 +2295,7 @@ int32 FDynamicMeshCollection::CutWithMultiplePlanes(
 			{
 				BoolResults.Add(MakeUnique<FMeshData>(NumUVLayers));
 				BoolResults[ResultIdx]->TransformIndex = TransformIndex;
-				BoolResults[ResultIdx]->ToCollection = ToCollection;
+				BoolResults[ResultIdx]->FromCollection = FromCollection;
 			}
 			check(CellMeshes.CellMeshes.Num() == 2);
 			bool bKeepResults = true;
@@ -2344,7 +2344,7 @@ int32 FDynamicMeshCollection::CutWithMultiplePlanes(
 						for (int32 Idx = 1; Idx < SplitMeshes.Num(); Idx++)
 						{
 							const FDynamicMesh3& Mesh = SplitMeshes[Idx];
-							ResultIndices.Add(ToCut.Add(MakeUnique<FMeshData>(Mesh, TransformIndex, ToCollection)));
+							ResultIndices.Add(ToCut.Add(MakeUnique<FMeshData>(Mesh, TransformIndex, FromCollection)));
 							ParentIndices.Add(UnsplitIdx);
 						}
 					}
@@ -2449,7 +2449,7 @@ int32 FDynamicMeshCollection::CutWithMultiplePlanes(
 			FDynamicMesh3& Mesh = ToCut[ToCutIdx]->AugMesh;
 
 			FString BoneName = GetBoneName(*Collection, ToCut[ToCutIdx]->TransformIndex, SubPartIdx++);
-			int32 CreatedGeometryIdx = AppendToCollection(ToCut[ToCutIdx]->ToCollection, Mesh, CollisionSampleSpacing, ToCut[ToCutIdx]->TransformIndex, BoneName, *Collection, InternalMaterialID);
+			int32 CreatedGeometryIdx = AppendToCollection(ToCut[ToCutIdx]->FromCollection, Mesh, CollisionSampleSpacing, ToCut[ToCutIdx]->TransformIndex, BoneName, *Collection, InternalMaterialID);
 			ToCutIdxToGeometryIdx[ToCutIdx] = CreatedGeometryIdx;
 			if (FirstCreatedIndex == -1)
 			{
@@ -2567,7 +2567,7 @@ int32 FDynamicMeshCollection::CutWithCellMeshes(const FInternalSurfaceMaterials&
 						{
 							FDynamicMesh3& Island = Islands[i];
 							FString BoneName = GetBoneName(*Collection, Surface.TransformIndex, SubPartIndex++);
-							CreatedGeometryIdx = AppendToCollection(Surface.ToCollection, Island, CollisionSampleSpacing, Surface.TransformIndex, BoneName, *Collection, InternalMaterialID);
+							CreatedGeometryIdx = AppendToCollection(Surface.FromCollection, Island, CollisionSampleSpacing, Surface.TransformIndex, BoneName, *Collection, InternalMaterialID);
 							CellToGeometry.Add(CellIdx, CreatedGeometryIdx);
 							if (i > 0)
 							{
@@ -2588,7 +2588,7 @@ int32 FDynamicMeshCollection::CutWithCellMeshes(const FInternalSurfaceMaterials&
 					else
 					{
 						FString BoneName = GetBoneName(*Collection, Surface.TransformIndex, SubPartIndex++);
-						CreatedGeometryIdx = AppendToCollection(Surface.ToCollection, AugBoolResult, CollisionSampleSpacing, Surface.TransformIndex, BoneName, *Collection, InternalMaterialID);
+						CreatedGeometryIdx = AppendToCollection(Surface.FromCollection, AugBoolResult, CollisionSampleSpacing, Surface.TransformIndex, BoneName, *Collection, InternalMaterialID);
 						CellToGeometry.Add(CellIdx, CreatedGeometryIdx);
 						GeometryToResultMesh.Add(CreatedGeometryIdx, CellIdx);
 
@@ -2824,7 +2824,7 @@ bool FDynamicMeshCollection::UpdateAllCollections(FGeometryCollection& Collectio
 		FDynamicMeshCollection::FMeshData& MeshData = Meshes[MeshIdx];
 		FDynamicMesh3& Mesh = MeshData.AugMesh;
 		int32 GeometryIdx = Collection.TransformToGeometryIndex[MeshData.TransformIndex];
-		bool bSucceeded = UpdateCollection(MeshData.ToCollection, Mesh, GeometryIdx, Collection, -1);
+		bool bSucceeded = UpdateCollection(MeshData.FromCollection, Mesh, GeometryIdx, Collection, -1);
 		bAllSucceeded &= bSucceeded;
 	}
 
@@ -2833,7 +2833,7 @@ bool FDynamicMeshCollection::UpdateAllCollections(FGeometryCollection& Collectio
 
 
 // Update an existing geometry in a collection w/ a new mesh (w/ the same number of faces and vertices!)
-bool FDynamicMeshCollection::UpdateCollection(const FTransform& ToCollection, FDynamicMesh3& Mesh, int32 GeometryIdx, FGeometryCollection& Output, int32 InternalMaterialID)
+bool FDynamicMeshCollection::UpdateCollection(const FTransform& FromCollection, FDynamicMesh3& Mesh, int32 GeometryIdx, FGeometryCollection& Output, int32 InternalMaterialID)
 {
 	if (!Mesh.IsCompact())
 	{
@@ -2862,8 +2862,8 @@ bool FDynamicMeshCollection::UpdateCollection(const FTransform& ToCollection, FD
 	{
 		checkSlow(Mesh.IsVertex(VID)); // mesh is compact
 		int32 CopyToIdx = VerticesStart + VID;
-		Output.Vertex[CopyToIdx] = ToCollection.TransformPosition(FVector(Mesh.GetVertex(VID)));
-		Output.Normal[CopyToIdx] = ToCollection.TransformVectorNoScale(FVector(Mesh.GetVertexNormal(VID)));
+		Output.Vertex[CopyToIdx] = FromCollection.InverseTransformPosition(FVector(Mesh.GetVertex(VID)));
+		Output.Normal[CopyToIdx] = FromCollection.InverseTransformVectorNoScale(FVector(Mesh.GetVertexNormal(VID)));
 		
 		for (int32 UVLayer = 0; UVLayer < UVLayerCount; ++UVLayer)
 		{
@@ -2874,8 +2874,8 @@ bool FDynamicMeshCollection::UpdateCollection(const FTransform& ToCollection, FD
 		
 		FVector3f TangentU, TangentV;
 		AugmentedDynamicMesh::GetTangent(Mesh, VID, TangentU, TangentV);
-		Output.TangentU[CopyToIdx] = ToCollection.TransformVectorNoScale(FVector(TangentU));
-		Output.TangentV[CopyToIdx] = ToCollection.TransformVectorNoScale(FVector(TangentV));
+		Output.TangentU[CopyToIdx] = FromCollection.InverseTransformVectorNoScale(FVector(TangentU));
+		Output.TangentV[CopyToIdx] = FromCollection.InverseTransformVectorNoScale(FVector(TangentV));
 		Output.Color[CopyToIdx] = FVector(Mesh.GetVertexColor(VID));
 
 		// Bone map is set based on the transform of the new geometry
@@ -2906,7 +2906,7 @@ bool FDynamicMeshCollection::UpdateCollection(const FTransform& ToCollection, FD
 }
 
 
-int32 FDynamicMeshCollection::AppendToCollection(const FTransform& ToCollection, FDynamicMesh3& Mesh, double CollisionSampleSpacing, int32 TransformParent, FString BoneName, FGeometryCollection& Output, int32 InternalMaterialID)
+int32 FDynamicMeshCollection::AppendToCollection(const FTransform& FromCollection, FDynamicMesh3& Mesh, double CollisionSampleSpacing, int32 TransformParent, FString BoneName, FGeometryCollection& Output, int32 InternalMaterialID)
 {
 	if (Mesh.TriangleCount() == 0)
 	{
@@ -2963,8 +2963,8 @@ int32 FDynamicMeshCollection::AppendToCollection(const FTransform& ToCollection,
 	{
 		checkSlow(Mesh.IsVertex(VID)); // mesh is compact
 		int32 CopyToIdx = VerticesStart + VID;
-		Output.Vertex[CopyToIdx] = ToCollection.TransformPosition(FVector(Mesh.GetVertex(VID)));
-		Output.Normal[CopyToIdx] = ToCollection.TransformVectorNoScale(FVector(Mesh.GetVertexNormal(VID)));
+		Output.Vertex[CopyToIdx] = FromCollection.InverseTransformPosition(FVector(Mesh.GetVertex(VID)));
+		Output.Normal[CopyToIdx] = FromCollection.InverseTransformVectorNoScale(FVector(Mesh.GetVertexNormal(VID)));
 		
 		Output.UVs[CopyToIdx].SetNum(UVLayerCount);
 		for (int32 UVLayer = 0; UVLayer < UVLayerCount; ++UVLayer)
@@ -2976,8 +2976,8 @@ int32 FDynamicMeshCollection::AppendToCollection(const FTransform& ToCollection,
 
 		FVector3f TangentU, TangentV;
 		AugmentedDynamicMesh::GetTangent(Mesh, VID, TangentU, TangentV);
-		Output.TangentU[CopyToIdx] = ToCollection.TransformVectorNoScale(FVector(TangentU));
-		Output.TangentV[CopyToIdx] = ToCollection.TransformVectorNoScale(FVector(TangentV));
+		Output.TangentU[CopyToIdx] = FromCollection.InverseTransformVectorNoScale(FVector(TangentU));
+		Output.TangentV[CopyToIdx] = FromCollection.InverseTransformVectorNoScale(FVector(TangentV));
 		Output.Color[CopyToIdx] = FVector(Mesh.GetVertexColor(VID));
 
 		// Bone map is set based on the transform of the new geometry
diff --git a/Engine/Plugins/Experimental/PlanarCutPlugin/Source/PlanarCut/Private/GeometryMeshConversion.h b/Engine/Plugins/Experimental/PlanarCutPlugin/Source/PlanarCut/Private/GeometryMeshConversion.h
--- a/Engine/Plugins/Experimental/PlanarCutPlugin/Source/PlanarCut/Private/GeometryMeshConversion.h
+++ b/Engine/Plugins/Experimental/PlanarCutPlugin/Source/PlanarCut/Private/GeometryMeshConversion.h
@@ -157,14 +157,14 @@ struct PLANARCUT_API FDynamicMeshCollection
 		// FDynamicMeshAABBTree3 Spatial; // TODO: maybe refactor mesh booleans to allow version where caller provides spatial data; it's computed every boolean now
 		// FTransform3d Transform; // TODO: maybe pretransform the data to a space that is good for cutting; refactor mesh boolean so there is an option to have it not transform input
 		int32 TransformIndex; // where the mesh was from in the geometry collection
-		FTransform ToCollection; // transform that need be applied to go back to the local space of the geometry collection
+		FTransform FromCollection; // transform that was used to go from the geometry collection to the local space used for processing
 
 		FMeshData(int32 NumUVLayers)
 		{
 			SetGeometryCollectionAttributes(AugMesh, NumUVLayers);
 		}
 
-		FMeshData(const UE::Geometry::FDynamicMesh3& Mesh, int32 TransformIndex, FTransform ToCollection) : AugMesh(Mesh), TransformIndex(TransformIndex), ToCollection(ToCollection)
+		FMeshData(const UE::Geometry::FDynamicMesh3& Mesh, int32 TransformIndex, FTransform FromCollection) : AugMesh(Mesh), TransformIndex(TransformIndex), FromCollection(FromCollection)
 		{}
 
 		void SetMesh(const UE::Geometry::FDynamicMesh3& NewAugMesh)
@@ -247,12 +247,12 @@ struct PLANARCUT_API FDynamicMeshCollection
 	// Resizes the GeometryCollection as needed
 	bool UpdateAllCollections(FGeometryCollection& Collection);
 
-	static int32 AppendToCollection(const FTransform& ToCollection, UE::Geometry::FDynamicMesh3& Mesh, double CollisionSampleSpacing, int32 TransformParent, FString BoneName, FGeometryCollection& Output, int32 InternalMaterialID);
+	static int32 AppendToCollection(const FTransform& FromCollection, UE::Geometry::FDynamicMesh3& Mesh, double CollisionSampleSpacing, int32 TransformParent, FString BoneName, FGeometryCollection& Output, int32 InternalMaterialID);
 
 private:
 
 	// Update an existing geometry in a collection w/ a new mesh (w/ the same number of faces and vertices!)
-	static bool UpdateCollection(const FTransform& ToCollection, UE::Geometry::FDynamicMesh3& Mesh, int32 GeometryIdx, FGeometryCollection& Output, int32 InternalMaterialID);
+	static bool UpdateCollection(const FTransform& FromCollection, UE::Geometry::FDynamicMesh3& Mesh, int32 GeometryIdx, FGeometryCollection& Output, int32 InternalMaterialID);
 
 	void FillVertexHash(const UE::Geometry::FDynamicMesh3& Mesh, UE::Geometry::TPointHashGrid3d<int>& VertHash);
 
diff --git a/Engine/Plugins/Experimental/PlanarCutPlugin/Source/PlanarCut/Private/PlanarCut.cpp b/Engine/Plugins/Experimental/PlanarCutPlugin/Source/PlanarCut/Private/PlanarCut.cpp
--- a/Engine/Plugins/Experimental/PlanarCutPlugin/Source/PlanarCut/Private/PlanarCut.cpp
+++ b/Engine/Plugins/Experimental/PlanarCutPlugin/Source/PlanarCut/Private/PlanarCut.cpp
@@ -1168,7 +1168,7 @@ int32 MergeBones(
 	for (int32 UpMeshIdx = 0; UpMeshIdx < UpdateCollection.Meshes.Num(); UpMeshIdx++)
 	{
 		FMeshData& UpMeshData = UpdateCollection.Meshes[UpMeshIdx];
-		FTransform3d UpMeshXF = (FTransform3d)UpMeshData.ToCollection;
+		FTransform3d UpMeshXF = (FTransform3d)UpMeshData.FromCollection;
 		int32 UpGeoIdx = Collection.TransformToGeometryIndex[UpMeshData.TransformIndex];
 		FRemoveGroup& Group = RemoveGroups[GeomIdxToRemoveGroupIdx[UpGeoIdx]];
 		if (!ensure(Group.IsValid()))
@@ -1179,7 +1179,7 @@ int32 MergeBones(
 		for (int32 RmGeoIdx : Group.ToRemove)
 		{
 			FMeshData& RmMeshData = RemoveCollection.Meshes[GeoIdxToRmMeshIdx[RmGeoIdx]];
-			FTransform3d RmMeshXF = (FTransform3d)RmMeshData.ToCollection;
+			FTransform3d RmMeshXF = (FTransform3d)RmMeshData.FromCollection;
 			if (bUnionJoinedPieces)
 			{
 				FMeshBoolean Boolean(&UpMeshData.AugMesh, &RmMeshData.AugMesh, &UpMeshData.AugMesh, FMeshBoolean::EBooleanOp::Union);
@@ -1193,11 +1193,11 @@ int32 MergeBones(
 				MeshEditor.AppendMesh(&RmMeshData.AugMesh, IndexMaps_Unused,
 					[&UpMeshXF, &RmMeshXF](int32 Unused, const FVector3d& Pos)
 					{
-						return UpMeshXF.InverseTransformPosition(RmMeshXF.TransformPosition(Pos));
+						return UpMeshXF.TransformPosition(RmMeshXF.InverseTransformPosition(Pos));
 					},
 					[&UpMeshXF, &RmMeshXF](int32 Unused, const FVector3d& Normal)
 					{
-						return UpMeshXF.InverseTransformNormal(RmMeshXF.TransformNormal(Normal));
+						return UpMeshXF.TransformNormal(RmMeshXF.InverseTransformNormal(Normal));
 					});
 			}
 		}
@@ -1296,13 +1296,13 @@ void ConvertToMeshDescription(
 	for (int32 MeshIdx = 0; MeshIdx < NumMeshes; MeshIdx++)
 	{
 		FDynamicMesh3& Mesh = MeshCollection.Meshes[MeshIdx].AugMesh;
-		const FTransform& ToCollection = MeshCollection.Meshes[MeshIdx].ToCollection;
+		const FTransform& FromCollection = MeshCollection.Meshes[MeshIdx].FromCollection;
 
 		// transform from the local geometry to the collection space
 		// unless it's just one mesh that we're going to center later anyway
 		if (!bCenterPivot || NumMeshes > 1)
 		{
-			MeshTransforms::ApplyTransform(Mesh, (UE::Geometry::FTransform3d)ToCollection);
+			MeshTransforms::ApplyTransformInverse(Mesh, (UE::Geometry::FTransform3d)FromCollection);
 		}
 
 		FMeshNormals::InitializeOverlayToPerVertexNormals(Mesh.Attributes()->PrimaryNormals(), true);
