fix incorrectly transformed geometry when fracturing a geometry collection that has non-uniform scale + rotation, by not using FTransform's Inverse() function

#rnx

diff --git a/Engine/Plugins/Experimental/PlanarCutPlugin/Source/PlanarCut/Private/GeometryMeshConversion.cpp b/Engine/Plugins/Experimental/PlanarCutPlugin/Source/PlanarCut/Private/GeometryMeshConversion.cpp
--- a/Engine/Plugins/Experimental/PlanarCutPlugin/Source/PlanarCut/Private/GeometryMeshConversion.cpp
+++ b/Engine/Plugins/Experimental/PlanarCutPlugin/Source/PlanarCut/Private/GeometryMeshConversion.cpp
@@ -2036,7 +2036,7 @@ void FDynamicMeshCollection::Init(const FGeometryCollection* Collection, const T
 		int32 AddedMeshIdx = Meshes.Add(new FMeshData(NumUVLayers));
 		FMeshData& MeshData = Meshes[AddedMeshIdx];
 		MeshData.TransformIndex = TransformIdx;
-		MeshData.ToCollection = FTransform(CollectionToLocal.Inverse());
+		MeshData.FromCollection = FTransform(CollectionToLocal);
 		FDynamicMesh3& Mesh = MeshData.AugMesh;
 		
 		Mesh.EnableAttributes();
@@ -2278,7 +2278,7 @@ int32 FDynamicMeshCollection::CutWithMultiplePlanes(
 		{
 			FMeshData& Surface = *ToCut[ToCutIdx];
 			int32 TransformIndex = Surface.TransformIndex;
-			FTransform ToCollection = Surface.ToCollection;
+			FTransform FromCollection = Surface.FromCollection;
 
 			FAxisAlignedBox3d Box = Surface.GetCachedBounds();
 			if (InternalSurfaceMaterials.NoiseSettings.IsSet())
@@ -2295,7 +2295,7 @@ int32 FDynamicMeshCollection::CutWithMultiplePlanes(
 			{
 				BoolResults.Add(MakeUnique<FMeshData>(NumUVLayers));
 				BoolResults[ResultIdx]->TransformIndex = TransformIndex;
-				BoolResults[ResultIdx]->ToCollection = ToCollection;
+				BoolResults[ResultIdx]->FromCollection = FromCollection;
 			}
 			check(CellMeshes.CellMeshes.Num() == 2);
 			bool bKeepResults = true;
@@ -2344,7 +2344,7 @@ int32 FDynamicMeshCollection::CutWithMultiplePlanes(
 						for (int32 Idx = 1; Idx < SplitMeshes.Num(); Idx++)
 						{
 							const FDynamicMesh3& Mesh = SplitMeshes[Idx];
-							ResultIndices.Add(ToCut.Add(MakeUnique<FMeshData>(Mesh, TransformIndex, ToCollection)));
+							ResultIndices.Add(ToCut.Add(MakeUnique<FMeshData>(Mesh, TransformIndex, FromCollection)));
 							ParentIndices.Add(UnsplitIdx);
 						}
 					}
@@ -2449,7 +2449,7 @@ int32 FDynamicMeshCollection::CutWithMultiplePlanes(
 			FDynamicMesh3& Mesh = ToCut[ToCutIdx]->AugMesh;
 
 			FString BoneName = GetBoneName(*Collection, ToCut[ToCutIdx]->TransformIndex, SubPartIdx++);
-			int32 CreatedGeometryIdx = AppendToCollection(ToCut[ToCutIdx]->ToCollection, Mesh, CollisionSampleSpacing, ToCut[ToCutIdx]->TransformIndex, BoneName, *Collection, InternalMaterialID);
+			int32 CreatedGeometryIdx = AppendToCollection(ToCut[ToCutIdx]->FromCollection, Mesh, CollisionSampleSpacing, ToCut[ToCutIdx]->TransformIndex, BoneName, *Collection, InternalMaterialID);
 			ToCutIdxToGeometryIdx[ToCutIdx] = CreatedGeometryIdx;
 			if (FirstCreatedIndex == -1)
 			{
@@ -2567,7 +2567,7 @@ int32 FDynamicMeshCollection::CutWithCellMeshes(const FInternalSurfaceMaterials&
 						{
 							FDynamicMesh3& Island = Islands[i];
 							FString BoneName = GetBoneName(*Collection, Surface.TransformIndex, SubPartIndex++);
-							CreatedGeometryIdx = AppendToCollection(Surface.ToCollection, Island, CollisionSampleSpacing, Surface.TransformIndex, BoneName, *Collection, InternalMaterialID);
+							CreatedGeometryIdx = AppendToCollection(Surface.FromCollection, Island, CollisionSampleSpacing, Surface.TransformIndex, BoneName, *Collection, InternalMaterialID);
 							CellToGeometry.Add(CellIdx, CreatedGeometryIdx);
 							if (i > 0)
 							{
@@ -2588,7 +2588,7 @@ int32 FDynamicMeshCollection::CutWithCellMeshes(const FInternalSurfaceMaterials&
 					else
 					{
 						FString BoneName = GetBoneName(*Collection, Surface.TransformIndex, SubPartIndex++);
-						CreatedGeometryIdx = AppendToCollection(Surface.ToCollection, AugBoolResult, CollisionSampleSpacing, Surface.TransformIndex, BoneName, *Collection, InternalMaterialID);
+						CreatedGeometryIdx = AppendToCollection(Surface.FromCollection, AugBoolResult, CollisionSampleSpacing, Surface.TransformIndex, BoneName, *Collection, InternalMaterialID);
 						CellToGeometry.Add(CellIdx, CreatedGeometryIdx);
 						GeometryToResultMesh.Add(CreatedGeometryIdx, CellIdx);
 
@@ -2824,7 +2824,7 @@ bool FDynamicMeshCollection::UpdateAllCollections(FGeometryCollection& Collectio
 		FDynamicMeshCollection::FMeshData& MeshData = Meshes[MeshIdx];
 		FDynamicMesh3& Mesh = MeshData.AugMesh;
 		int32 GeometryIdx = Collection.TransformToGeometryIndex[MeshData.TransformIndex];
-		bool bSucceeded = UpdateCollection(MeshData.ToCollection, Mesh, GeometryIdx, Collection, -1);
+		bool bSucceeded = UpdateCollection(MeshData.FromCollection, Mesh, GeometryIdx, Collection, -1);
 		bAllSucceeded &= bSucceeded;
 	}
 
@@ -2833,7 +2833,7 @@ bool FDynamicMeshCollection::UpdateAllCollections(FGeometryCollection& Collectio
 
 
 // Update an existing geometry in a collection w/ a new mesh (w/ the same number of faces and vertices!)
-bool FDynamicMeshCollection::UpdateCollection(const FTransform& ToCollection, FDynamicMesh3& Mesh, int32 GeometryIdx, FGeometryCollection& Output, int32 InternalMaterialID)
+bool FDynamicMeshCollection::UpdateCollection(const FTransform& FromCollection, FDynamicMesh3& Mesh, int32 GeometryIdx, FGeometryCollection& Output, int32 InternalMaterialID)
 {
 	if (!Mesh.IsCompact())
 	{
@@ -2862,8 +2862,8 @@ bool FDynamicMeshCollection::UpdateCollection(const FTransform& ToCollection, FD
 	{
 		checkSlow(Mesh.IsVertex(VID)); // mesh is compact
 		int32 CopyToIdx = VerticesStart + VID;
-		Output.Vertex[CopyToIdx] = ToCollection.TransformPosition(FVector(Mesh.GetVertex(VID)));
-		Output.Normal[CopyToIdx] = ToCollection.TransformVectorNoScale(FVector(Mesh.GetVertexNormal(VID)));
+		Output.Vertex[CopyToIdx] = FromCollection.InverseTransformPosition(FVector(Mesh.GetVertex(VID)));
+		Output.Normal[CopyToIdx] = FromCollection.InverseTransformVectorNoScale(FVector(Mesh.GetVertexNormal(VID)));
 		
 		for (int32 UVLayer = 0; UVLayer < UVLayerCount; ++UVLayer)
 		{
@@ -2874,8 +2874,8 @@ bool FDynamicMeshCollection::UpdateCollection(const FTransform& ToCollection, FD
 		
 		FVector3f TangentU, TangentV;
 		AugmentedDynamicMesh::GetTangent(Mesh, VID, TangentU, TangentV);
-		Output.TangentU[CopyToIdx] = ToCollection.TransformVectorNoScale(FVector(TangentU));
-		Output.TangentV[CopyToIdx] = ToCollection.TransformVectorNoScale(FVector(TangentV));
+		Output.TangentU[CopyToIdx] = FromCollection.InverseTransformVectorNoScale(FVector(TangentU));
+		Output.TangentV[CopyToIdx] = FromCollection.InverseTransformVectorNoScale(FVector(TangentV));
 		Output.Color[CopyToIdx] = FVector(Mesh.GetVertexColor(VID));
 
 		// Bone map is set based on the transform of the new geometry
@@ -2906,7 +2906,7 @@ bool FDynamicMeshCollection::UpdateCollection(const FTransform& ToCollection, FD
 }
 
 
-int32 FDynamicMeshCollection::AppendToCollection(const FTransform& ToCollection, FDynamicMesh3& Mesh, double CollisionSampleSpacing, int32 TransformParent, FString BoneName, FGeometryCollection& Output, int32 InternalMaterialID)
+int32 FDynamicMeshCollection::AppendToCollection(const FTransform& FromCollection, FDynamicMesh3& Mesh, double CollisionSampleSpacing, int32 TransformParent, FString BoneName, FGeometryCollection& Output, int32 InternalMaterialID)
 {
 	if (Mesh.TriangleCount() == 0)
 	{
@@ -2963,8 +2963,8 @@ int32 FDynamicMeshCollection::AppendToCollection(const FTransform& ToCollection,
 	{
 		checkSlow(Mesh.IsVertex(VID)); // mesh is compact
 		int32 CopyToIdx = VerticesStart + VID;
-		Output.Vertex[CopyToIdx] = ToCollection.TransformPosition(FVector(Mesh.GetVertex(VID)));
-		Output.Normal[CopyToIdx] = ToCollection.TransformVectorNoScale(FVector(Mesh.GetVertexNormal(VID)));
+		Output.Vertex[CopyToIdx] = FromCollection.InverseTransformPosition(FVector(Mesh.GetVertex(VID)));
+		Output.Normal[CopyToIdx] = FromCollection.InverseTransformVectorNoScale(FVector(Mesh.GetVertexNormal(VID)));
 		
 		Output.UVs[CopyToIdx].SetNum(UVLayerCount);
 		for (int32 UVLayer = 0; UVLayer < UVLayerCount; ++UVLayer)
@@ -2976,8 +2976,8 @@ int32 FDynamicMeshCollection::AppendToCollection(const FTransform& ToCollection,
 
 		FVector3f TangentU, TangentV;
 		AugmentedDynamicMesh::GetTangent(Mesh, VID, TangentU, TangentV);
-		Output.TangentU[CopyToIdx] = ToCollection.TransformVectorNoScale(FVector(TangentU));
-		Output.TangentV[CopyToIdx] = ToCollection.TransformVectorNoScale(FVector(TangentV));
+		Output.TangentU[CopyToIdx] = FromCollection.InverseTransformVectorNoScale(FVector(TangentU));
+		Output.TangentV[CopyToIdx] = FromCollection.InverseTransformVectorNoScale(FVector(TangentV));
 		Output.Color[CopyToIdx] = FVector(Mesh.GetVertexColor(VID));
 
 		// Bone map is set based on the transform of the new geometry
