add normals option to fracture autouv

diff --git a/Engine/Plugins/Experimental/ChaosEditor/Source/FractureEditor/Private/FractureToolUV.cpp b/Engine/Plugins/Experimental/ChaosEditor/Source/FractureEditor/Private/FractureToolUV.cpp
--- a/Engine/Plugins/Experimental/ChaosEditor/Source/FractureEditor/Private/FractureToolUV.cpp
+++ b/Engine/Plugins/Experimental/ChaosEditor/Source/FractureEditor/Private/FractureToolUV.cpp
@@ -441,7 +441,7 @@ void UFractureToolAutoUV::BakeTextureForComponent(UGeometryCollectionComponent*
 			0
 		);
 	}
-	else
+	else if (AutoUVSettings->BakeTextureType == ETextureType::SpatialGradients)
 	{
 		Attributes = FIndex4i(
 			(int32)EBakeAttributes::PositionX,
@@ -450,6 +450,15 @@ void UFractureToolAutoUV::BakeTextureForComponent(UGeometryCollectionComponent*
 			0
 		);
 	}
+	else // ETextureType::Normals
+	{
+		Attributes = FIndex4i(
+			(int32)EBakeAttributes::NormalX,
+			(int32)EBakeAttributes::NormalY,
+			(int32)EBakeAttributes::NormalZ,
+			0
+		);
+	}
 	UE::PlanarCut::FTextureAttributeSettings AttribSettings;
 	AttribSettings.ToExternal_MaxDistance = AutoUVSettings->MaxDistance;
 	AttribSettings.AO_Rays = AutoUVSettings->OcclusionRays;
@@ -476,6 +485,10 @@ void UFractureToolAutoUV::BakeTextureForComponent(UGeometryCollectionComponent*
 	{
 		Suffix = "_AutoUV_Spatial";
 	}
+	else if (AutoUVSettings->BakeTextureType == ETextureType::Normals)
+	{
+		Suffix = "_AutoUV_Normals";
+	}
 	Suffix = FPaths::MakeValidFileName(Suffix);
 	SaveGeneratedTexture(ImageBuilder, FString::Printf(TEXT("%s%s"), *BaseName, *Suffix), Component->GetRestCollection(), AutoUVSettings->bPromptToSave, AutoUVSettings->bReplaceExisting);
 }
diff --git a/Engine/Plugins/Experimental/ChaosEditor/Source/FractureEditor/Private/FractureToolUV.h b/Engine/Plugins/Experimental/ChaosEditor/Source/FractureEditor/Private/FractureToolUV.h
--- a/Engine/Plugins/Experimental/ChaosEditor/Source/FractureEditor/Private/FractureToolUV.h
+++ b/Engine/Plugins/Experimental/ChaosEditor/Source/FractureEditor/Private/FractureToolUV.h
@@ -29,8 +29,12 @@ enum class EAutoUVTextureResolution : int32
 UENUM()
 enum class ETextureType
 {
+	// Bake thickness, ambient occlusion, and a smoothed curvature metric
 	ThicknessAndSurfaceAttributes,
-	SpatialGradients
+	// Bake position relative to the bounding box
+	SpatialGradients,
+	// Bake surface normals
+	Normals
 };
 
 UENUM()
diff --git a/Engine/Plugins/Experimental/PlanarCutPlugin/Source/PlanarCut/Private/FractureAutoUV.cpp b/Engine/Plugins/Experimental/PlanarCutPlugin/Source/PlanarCut/Private/FractureAutoUV.cpp
--- a/Engine/Plugins/Experimental/PlanarCutPlugin/Source/PlanarCut/Private/FractureAutoUV.cpp
+++ b/Engine/Plugins/Experimental/PlanarCutPlugin/Source/PlanarCut/Private/FractureAutoUV.cpp
@@ -600,6 +600,8 @@ void TextureInternalSurfaces(
 	int DistanceToExternalIdx = BakeAttributes.IndexOf((int)EBakeAttributes::DistanceToExternal);
 	int AmbientIdx = BakeAttributes.IndexOf((int)EBakeAttributes::AmbientOcclusion);
 	int CurvatureIdx = BakeAttributes.IndexOf((int)EBakeAttributes::Curvature);
+	int NormalXIdx = BakeAttributes.IndexOf((int)EBakeAttributes::NormalX);
+	int NormalYIdx = BakeAttributes.IndexOf((int)EBakeAttributes::NormalY);
 	int NormalZIdx = BakeAttributes.IndexOf((int)EBakeAttributes::NormalZ);
 
 	FAxisAlignedBox3d Box = OutsideSpatial.GetBoundingBox();
@@ -768,7 +770,7 @@ void TextureInternalSurfaces(
 
 	ParallelFor(OccupancyMap.Dimensions.GetHeight(),
 		[&AttributeSettings, &TextureOut, &UVMesh, &OccupancyMap, &ToTextureMesh, &OutsideSpatial,
-		 &DistanceToExternalIdx, &AmbientIdx, &NormalZIdx, &PosXIdx, &PosYIdx, &PosZIdx](int32 Y)
+		 &DistanceToExternalIdx, &AmbientIdx, &NormalXIdx, &NormalYIdx, &NormalZIdx, &PosXIdx, &PosYIdx, &PosZIdx](int32 Y)
 	{
 		for (int32 X = 0; X < OccupancyMap.Dimensions.GetWidth(); X++)
 		{
@@ -784,7 +786,7 @@ void TextureInternalSurfaces(
 				FDistPoint3Triangle3d DistQuery = TMeshQueries<FGeomFlatUVMesh>::TriangleDistance(UVMesh, TID, FVector3d(UV.X, UV.Y, 0));
 				FVector3d Bary = DistQuery.TriangleBaryCoords;
 				FVector3f Normal(0, 0, 1);
-				bool bNeedsNormal = NormalZIdx > -1 || AmbientIdx > -1;
+				bool bNeedsNormal = NormalXIdx > -1 || NormalYIdx > -1 || NormalZIdx > -1 || AmbientIdx > -1;
 				if (bNeedsNormal)
 				{
 					Normal = ToTextureMesh.GetInterpolatedNormal(TID, Bary);
@@ -818,19 +820,20 @@ void TextureInternalSurfaces(
 					SetSpatial(PosYIdx, 1);
 					SetSpatial(PosZIdx, 2);
 				}
-				if (NormalZIdx > -1)
+				auto SetNormal = [&Normal, &OutColor](int TargetIdx, int Dim)
 				{
-					float NormalZ = Normal.Z;
-					if (AttributeSettings.bNormalZ_TakeAbs)
+					if (TargetIdx > -1)
 					{
-						NormalZ = FMathf::Abs(NormalZ);
-					}
-					else
-					{
-						NormalZ = (1 + NormalZ) * .5; // compress from [-1,1] to [0,1] range
+						float NormalVal = Normal[Dim];
+						{
+							NormalVal = (1 + NormalVal) * .5f; // compress from [-1,1] to [0,1] range
+						}
+						OutColor[TargetIdx] = NormalVal;
 					}
-					OutColor[NormalZIdx] = NormalZ;
-				}
+				};
+				SetNormal(NormalXIdx, 0);
+				SetNormal(NormalYIdx, 1);
+				SetNormal(NormalZIdx, 2);
 
 				TextureOut.SetPixel(LinearCoord, OutColor);
 			}
diff --git a/Engine/Plugins/Experimental/PlanarCutPlugin/Source/PlanarCut/Public/FractureAutoUV.h b/Engine/Plugins/Experimental/PlanarCutPlugin/Source/PlanarCut/Public/FractureAutoUV.h
--- a/Engine/Plugins/Experimental/PlanarCutPlugin/Source/PlanarCut/Public/FractureAutoUV.h
+++ b/Engine/Plugins/Experimental/PlanarCutPlugin/Source/PlanarCut/Public/FractureAutoUV.h
@@ -61,6 +61,8 @@ enum class EBakeAttributes : int32
 	DistanceToExternal,
 	AmbientOcclusion,
 	Curvature,
+	NormalX,
+	NormalY,
 	NormalZ,
 	PositionX,
 	PositionY,
@@ -84,8 +86,6 @@ struct FTextureAttributeSettings
 	double Curvature_ThicknessFactor = 3.0; // distance to search for mesh correspondence, as a factor of voxel size
 	double Curvature_MaxValue = .1; // curvatures above this value will be clamped
 	int ClearGutterChannel = -1; // don't copy gutter values for this channel, if specified -- useful for visualizing the UV island borders
-	
-	bool bNormalZ_TakeAbs = true;
 };
 
 
