let fracture normals tool optionally recompute the normal topology too

diff --git a/Engine/Plugins/Experimental/ChaosEditor/Source/FractureEditor/Private/FractureToolRecomputeNormals.cpp b/Engine/Plugins/Experimental/ChaosEditor/Source/FractureEditor/Private/FractureToolRecomputeNormals.cpp
--- a/Engine/Plugins/Experimental/ChaosEditor/Source/FractureEditor/Private/FractureToolRecomputeNormals.cpp
+++ b/Engine/Plugins/Experimental/ChaosEditor/Source/FractureEditor/Private/FractureToolRecomputeNormals.cpp
@@ -135,7 +135,7 @@ int32 UFractureToolRecomputeNormals::ExecuteFracture(const FFractureToolContext&
 		const UFractureCollisionSettings* LocalCutSettings = CollisionSettings;
 
 		bool bOnlyInternalFaces = true;
-		RecomputeNormalsAndTangents(NormalsSettings->bOnlyTangents, *FractureContext.GetGeometryCollection(), FractureContext.GetSelection(), NormalsSettings->bOnlyInternalSurfaces);
+		RecomputeNormalsAndTangents(NormalsSettings->bOnlyTangents, NormalsSettings->bRecomputeSharpEdges, NormalsSettings->SharpEdgeAngleThreshold, *FractureContext.GetGeometryCollection(), FractureContext.GetSelection(), NormalsSettings->bOnlyInternalSurfaces);
 	}
 
 	return INDEX_NONE;
diff --git a/Engine/Plugins/Experimental/ChaosEditor/Source/FractureEditor/Private/FractureToolRecomputeNormals.h b/Engine/Plugins/Experimental/ChaosEditor/Source/FractureEditor/Private/FractureToolRecomputeNormals.h
--- a/Engine/Plugins/Experimental/ChaosEditor/Source/FractureEditor/Private/FractureToolRecomputeNormals.h
+++ b/Engine/Plugins/Experimental/ChaosEditor/Source/FractureEditor/Private/FractureToolRecomputeNormals.h
@@ -37,6 +37,14 @@ public:
 	UPROPERTY(EditAnywhere, Category = RecomputeSettings)
 	bool bOnlyTangents = false;
 
+	/** If true, update where edges are 'sharp' by comparing adjacent triangle face normals vs the Sharp Edge Angle Threshold. */
+	UPROPERTY(EditAnywhere, Category = RecomputeSettings, meta = (EditCondition = "!bOnlyTangents"))
+	bool bRecomputeSharpEdges = false;
+
+	/** Threshold on angle of change in face normals across an edge, above which we create a sharp edge if bRecomputeSharpEdges is true */
+	UPROPERTY(EditAnywhere, Category = RecomputeSettings, meta = (UIMin = "0.0", UIMax = "180.0", ClampMin = "0.0", ClampMax = "180.0", EditCondition = "bRecomputeSharpEdges && !bOnlyTangents"))
+	float SharpEdgeAngleThreshold = 60.0f;
+
 	/** Whether to only change internal surface normals / tangents */
 	UPROPERTY(EditAnywhere, Category = RecomputeSettings, AdvancedDisplay)
 	bool bOnlyInternalSurfaces = true;
diff --git a/Engine/Plugins/Experimental/PlanarCutPlugin/Source/PlanarCut/Private/GeometryMeshConversion.cpp b/Engine/Plugins/Experimental/PlanarCutPlugin/Source/PlanarCut/Private/GeometryMeshConversion.cpp
--- a/Engine/Plugins/Experimental/PlanarCutPlugin/Source/PlanarCut/Private/GeometryMeshConversion.cpp
+++ b/Engine/Plugins/Experimental/PlanarCutPlugin/Source/PlanarCut/Private/GeometryMeshConversion.cpp
@@ -442,19 +442,90 @@ namespace AugmentedDynamicMesh
 		}
 	}
 
-	void ComputeTangents(FDynamicMesh3& Mesh, bool bOnlyOddMaterials, const TArrayView<const int32>& WhichMaterials, bool bRecomputeNormals)
+	void ComputeTangents(FDynamicMesh3& Mesh, bool bOnlyOddMaterials, const TArrayView<const int32>& WhichMaterials,
+		bool bRecomputeNormals, bool bMakeSharpEdges, float SharpAngleDegrees)
 	{
+		bMakeSharpEdges = bMakeSharpEdges && bRecomputeNormals; // cannot make sharp edges if normals aren't supposed to change
+
 		FDynamicMeshNormalOverlay* Normals = Mesh.Attributes()->PrimaryNormals();
-		FMeshNormals::InitializeOverlayToPerVertexNormals(Normals, !bRecomputeNormals);
-		if (bRecomputeNormals)
-		{
-			FMeshNormals::QuickRecomputeOverlayNormals(Mesh);
-		}
+		FMeshNormals::InitializeOverlayToPerVertexNormals(Normals, !bRecomputeNormals || bMakeSharpEdges);
 
 		// Copy per-vertex UVs to a UV overlay, because that's what the tangents code uses
 		// (TODO: consider making a tangent computation path that uses vertex normals / UVs)
 		InitializeOverlayToPerVertexUVs(Mesh, 1); // tangents only use the first layer
 		FDynamicMeshUVOverlay* UVs = Mesh.Attributes()->PrimaryUV();
+		FDynamicMeshMaterialAttribute* MaterialIDs = Mesh.Attributes()->GetMaterialID();
+
+		auto ShouldUpdateMID = [bOnlyOddMaterials, &WhichMaterials](int MaterialID)
+		{
+			bool bSkipEvenMaterial = bOnlyOddMaterials && ((MaterialID % 2) == 0);
+			return !bSkipEvenMaterial && (WhichMaterials.IsEmpty() || WhichMaterials.Contains(MaterialID));
+		};
+
+		// To update the normals topology, we need to weld and re-split the whole mesh
+		if (bMakeSharpEdges)
+		{
+			FDynamicMeshNormalOverlay NormalsOverlayCopy(&Mesh);
+
+			// Apply coincident edge merge so that sharp normals can be smoothed where the edge angle is below the threshold
+			AugmentedDynamicMesh::InitializeOverlayToPerVertexTangents(Mesh);
+			FMergeCoincidentMeshEdges EdgeWelder(&Mesh);
+			EdgeWelder.Apply();
+
+			NormalsOverlayCopy.Copy(*Normals);
+			// need to re-topo where the materials match WhichMaterials
+			double NormalDotProdThreshold = FMathd::Cos(SharpAngleDegrees * FMathd::DegToRad);
+
+			FMeshNormals FaceNormals(&Mesh);
+			FaceNormals.ComputeTriangleNormals();
+			Normals->CreateFromPredicate([&](int VID, int TA, int TB) {
+				int MA = MaterialIDs->GetValue(TA), MB = MaterialIDs->GetValue(TB);
+				if ((MA % 2) != (MB % 2))
+				{
+					return false; // always split at an internal/external face boundary
+				}
+				bool bShouldUpdateTopo = ShouldUpdateMID(MA) && ShouldUpdateMID(MB);
+				if (bShouldUpdateTopo) // in the region we're updating, don't split above dot threshold
+				{
+					return FaceNormals[TA].Dot(FaceNormals[TB]) > NormalDotProdThreshold;
+				}
+				else // in the region we're not updating, keep the original connectivity
+				{
+					return NormalsOverlayCopy.AreTrianglesConnected(TA, TB);
+				}
+			}, 0);
+
+			// copy back recomputed normals to the triangles that need updating, original normals to the triangles we don't want to change
+			FMeshNormals RecomputedNormals(&Mesh);
+			RecomputedNormals.RecomputeOverlayNormals(Normals);
+			for (int TID : Mesh.TriangleIndicesItr())
+			{
+				FIndex3i OverlayTri = Normals->GetTriangle(TID);
+				if (ShouldUpdateMID(MaterialIDs->GetValue(TID)))
+				{
+					for (int SubIdx = 0; SubIdx < 3; SubIdx++)
+					{
+						int ElementID = OverlayTri[SubIdx];
+						Normals->SetElement(ElementID, RecomputedNormals[ElementID]);
+					}
+				}
+				else
+				{
+					FIndex3i OldOverlayTri = NormalsOverlayCopy.GetTriangle(TID);
+					for (int SubIdx = 0; SubIdx < 3; SubIdx++)
+					{
+						int NewElementID = OverlayTri[SubIdx];
+						int OldElementID = OldOverlayTri[SubIdx];
+						Normals->SetElement(NewElementID, NormalsOverlayCopy.GetElement(OldElementID));
+					}
+				}
+			}
+
+			Mesh.CompactInPlace();
+
+			// Re-split vertices with different UVs/normals/tangents and transfer attributes back (required because we weld edges above)
+			AugmentedDynamicMesh::SplitOverlayAttributesToPerVertex(Mesh, true, true);
+		}
 
 		FComputeTangentsOptions Options;
 		Options.bAngleWeighted = true;
@@ -464,26 +535,31 @@ namespace AugmentedDynamicMesh
 
 		const TArray<FVector3f>& TanU = Tangents.GetTangents();
 		const TArray<FVector3f>& TanV = Tangents.GetBitangents();
-		FDynamicMeshMaterialAttribute* MaterialIDs = Mesh.Attributes()->GetMaterialID();
 		for (int TID : Mesh.TriangleIndicesItr())
 		{
-			int MaterialID = MaterialIDs->GetValue(TID);
-			if (bOnlyOddMaterials && (MaterialID % 2) == 0)
-			{
-				continue;
-			}
-			else if (WhichMaterials.Contains(MaterialID))
+			if (!ShouldUpdateMID(MaterialIDs->GetValue(TID)))
 			{
 				continue;
 			}
 					
 			int TanIdxBase = TID * 3;
 			FIndex3i Tri = Mesh.GetTriangle(TID);
+			
+			FIndex3i NormalElTri;
+			if (!bMakeSharpEdges) // if sharp edges, normals were already copied to the vertices, above
+			{
+				NormalElTri = Normals->GetTriangle(TID);
+			}
 			for (int Idx = 0; Idx < 3; Idx++)
 			{
 				int VID = Tri[Idx];
 				int TanIdx = TanIdxBase + Idx;
-				SetTangent(Mesh, VID, Mesh.GetVertexNormal(VID), TanU[TanIdx], TanV[TanIdx]);
+				FVector3f VertexNormal = bMakeSharpEdges ? Mesh.GetVertexNormal(VID) : Normals->GetElement(NormalElTri[Idx]);
+				if (!bMakeSharpEdges && bRecomputeNormals)
+				{
+					Mesh.SetVertexNormal(VID, VertexNormal);
+				}
+				SetTangent(Mesh, VID, VertexNormal, TanU[TanIdx], TanV[TanIdx]);
 			}
 		}
 	}
diff --git a/Engine/Plugins/Experimental/PlanarCutPlugin/Source/PlanarCut/Private/GeometryMeshConversion.h b/Engine/Plugins/Experimental/PlanarCutPlugin/Source/PlanarCut/Private/GeometryMeshConversion.h
--- a/Engine/Plugins/Experimental/PlanarCutPlugin/Source/PlanarCut/Private/GeometryMeshConversion.h
+++ b/Engine/Plugins/Experimental/PlanarCutPlugin/Source/PlanarCut/Private/GeometryMeshConversion.h
@@ -37,7 +37,8 @@ namespace AugmentedDynamicMesh
 	/// Initialize UV overlays based on the custom AugmentedDynamicMesh per-vertex UV attributes.  Optionally use FirstUVLayer to skip layers
 	void PLANARCUT_API InitializeOverlayToPerVertexUVs(UE::Geometry::FDynamicMesh3& Mesh, int32 NumUVLayers, int32 FirstUVLayer = 0);
 	void PLANARCUT_API InitializeOverlayToPerVertexTangents(UE::Geometry::FDynamicMesh3& Mesh);
-	void PLANARCUT_API ComputeTangents(UE::Geometry::FDynamicMesh3& Mesh, bool bOnlyOddMaterials, const TArrayView<const int32>& WhichMaterials, bool bRecomputeNormals = true);
+	void PLANARCUT_API ComputeTangents(UE::Geometry::FDynamicMesh3& Mesh, bool bOnlyOddMaterials, const TArrayView<const int32>& WhichMaterials,
+		bool bRecomputeNormals = true, bool bMakeSharpEdges = false, float SharpAngleDegrees = 60);
 	void PLANARCUT_API AddCollisionSamplesPerComponent(UE::Geometry::FDynamicMesh3& Mesh, double Spacing);
 	void PLANARCUT_API SplitOverlayAttributesToPerVertex(UE::Geometry::FDynamicMesh3& Mesh, bool bSplitUVs = true, bool bSplitTangents = true);
 }
diff --git a/Engine/Plugins/Experimental/PlanarCutPlugin/Source/PlanarCut/Private/PlanarCut.cpp b/Engine/Plugins/Experimental/PlanarCutPlugin/Source/PlanarCut/Private/PlanarCut.cpp
--- a/Engine/Plugins/Experimental/PlanarCutPlugin/Source/PlanarCut/Private/PlanarCut.cpp
+++ b/Engine/Plugins/Experimental/PlanarCutPlugin/Source/PlanarCut/Private/PlanarCut.cpp
@@ -1238,7 +1238,7 @@ int32 MergeBones(
 }
 
 
-void RecomputeNormalsAndTangents(bool bOnlyTangents, FGeometryCollection& Collection, const TArrayView<const int32>& TransformIndices,
+void RecomputeNormalsAndTangents(bool bOnlyTangents, bool bMakeSharpEdges, float SharpAngleDegrees, FGeometryCollection& Collection, const TArrayView<const int32>& TransformIndices,
 								 bool bOnlyOddMaterials, const TArrayView<const int32>& WhichMaterials)
 {
 	FTransform CellsToWorld = FTransform::Identity;
@@ -1248,7 +1248,7 @@ void RecomputeNormalsAndTangents(bool bOnlyTangents, FGeometryCollection& Collec
 	for (int MeshIdx = 0; MeshIdx < MeshCollection.Meshes.Num(); MeshIdx++)
 	{
 		FDynamicMesh3& Mesh = MeshCollection.Meshes[MeshIdx].AugMesh;
-		AugmentedDynamicMesh::ComputeTangents(Mesh, bOnlyOddMaterials, WhichMaterials, !bOnlyTangents);
+		AugmentedDynamicMesh::ComputeTangents(Mesh, bOnlyOddMaterials, WhichMaterials, !bOnlyTangents, bMakeSharpEdges, SharpAngleDegrees);
 	}
 
 	MeshCollection.UpdateAllCollections(Collection);
diff --git a/Engine/Plugins/Experimental/PlanarCutPlugin/Source/PlanarCut/Public/PlanarCut.h b/Engine/Plugins/Experimental/PlanarCutPlugin/Source/PlanarCut/Public/PlanarCut.h
--- a/Engine/Plugins/Experimental/PlanarCutPlugin/Source/PlanarCut/Public/PlanarCut.h
+++ b/Engine/Plugins/Experimental/PlanarCutPlugin/Source/PlanarCut/Public/PlanarCut.h
@@ -281,12 +281,14 @@ int32 PLANARCUT_API MergeBones(
  * Recompute normals and tangents of selected geometry, optionally restricted to faces with odd or given material IDs (i.e. to target internal faces)
  *
  * @param bOnlyTangents		If true, leave normals unchanged and only recompute tangent&bitangent vectors
+ * @param bMakeSharpEdges	If true, recompute the normal topology to split normals at 'sharp' edges
+ * @param SharpAngleDegrees	If bMakeSharpEdges, edges w/ adjacent triangle normals deviating by more than this threshold will be sharp edges (w/ split normals)
  * @param Collection		The Geometry Collection to be updated
  * @param TransformIndices	Which transform groups on the Geometry Collection to be updated.  If empty, all groups are updated.
  * @param bOnlyOddMaterials	If true, restrict recomputation to odd-numbered material IDs
  * @param WhichMaterials	If non-empty, restrict recomputation to only the listed material IDs
  */
-void PLANARCUT_API RecomputeNormalsAndTangents(bool bOnlyTangents, FGeometryCollection& Collection, const TArrayView<const int32>& TransformIndices = TArrayView<const int32>(),
+void PLANARCUT_API RecomputeNormalsAndTangents(bool bOnlyTangents, bool bMakeSharpEdges, float SharpAngleDegrees, FGeometryCollection& Collection, const TArrayView<const int32>& TransformIndices = TArrayView<const int32>(),
 	bool bOnlyOddMaterials = true, const TArrayView<const int32>& WhichMaterials = TArrayView<const int32>());
 /**
  * Scatter additional vertices (w/ no associated triangle) as needed to satisfy minimum point spacing
