Fix fracture normals tool crash (fix handling of overlay elements w/ no parent vertex)

#rnx

diff --git a/Engine/Plugins/Experimental/PlanarCutPlugin/Source/PlanarCut/Private/GeometryMeshConversion.cpp b/Engine/Plugins/Experimental/PlanarCutPlugin/Source/PlanarCut/Private/GeometryMeshConversion.cpp
--- a/Engine/Plugins/Experimental/PlanarCutPlugin/Source/PlanarCut/Private/GeometryMeshConversion.cpp
+++ b/Engine/Plugins/Experimental/PlanarCutPlugin/Source/PlanarCut/Private/GeometryMeshConversion.cpp
@@ -357,7 +357,11 @@ namespace AugmentedDynamicMesh
 				static_cast<TDynamicMeshVertexAttribute<float, 2>*>(Mesh.Attributes()->GetAttachedAttribute(UVChannelNames[LayerIdx]));
 			for (int32 EID : UVOverlay->ElementIndicesItr())
 			{
-				UVs->SetValue(UVOverlay->GetParentVertex(EID), UVOverlay->GetElement(EID));
+				int32 ParentVID = UVOverlay->GetParentVertex(EID);
+				if (ParentVID != INDEX_NONE)
+				{
+					UVs->SetValue(ParentVID, UVOverlay->GetElement(EID));
+				}
 			}
 		}
 		if (Mesh.Attributes()->NumNormalLayers() > 0)
@@ -365,7 +369,11 @@ namespace AugmentedDynamicMesh
 			FDynamicMeshNormalOverlay* Overlay = Mesh.Attributes()->GetNormalLayer(0);
 			for (int32 EID : Overlay->ElementIndicesItr())
 			{
-				Mesh.SetVertexNormal(Overlay->GetParentVertex(EID), Overlay->GetElement(EID));
+				int32 ParentVID = Overlay->GetParentVertex(EID);
+				if (ParentVID != INDEX_NONE)
+				{
+					Mesh.SetVertexNormal(ParentVID, Overlay->GetElement(EID));
+				}
 			}
 		}
 		if (Mesh.Attributes()->HasTangentSpace())
@@ -381,7 +389,11 @@ namespace AugmentedDynamicMesh
 				TDynamicMeshVertexAttribute<float, 3>* Tangent = Tangents[Idx - 1];
 				for (int32 EID : Overlay->ElementIndicesItr())
 				{
-					Tangent->SetValue(Overlay->GetParentVertex(EID), Overlay->GetElement(EID));
+					int32 ParentVID = Overlay->GetParentVertex(EID);
+					if (ParentVID != INDEX_NONE)
+					{
+						Tangent->SetValue(ParentVID, Overlay->GetElement(EID));
+					}
 				}
 			}
 		}
@@ -452,7 +464,9 @@ namespace AugmentedDynamicMesh
 
 		// Copy per-vertex UVs to a UV overlay, because that's what the tangents code uses
 		// (TODO: consider making a tangent computation path that uses vertex normals / UVs)
-		InitializeOverlayToPerVertexUVs(Mesh, 1); // tangents only use the first layer
+		// only need 1 UV layer unless we're round-tripping all the attributes through the mesh
+		int32 NeedNumUVLayers = bMakeSharpEdges ? NumEnabledUVChannels(Mesh) : 1;
+		InitializeOverlayToPerVertexUVs(Mesh, NeedNumUVLayers, 0);
 		FDynamicMeshUVOverlay* UVs = Mesh.Attributes()->PrimaryUV();
 		FDynamicMeshMaterialAttribute* MaterialIDs = Mesh.Attributes()->GetMaterialID();
 
