make fracture noise consistent for a given random seed
also use fracturecontext random seed + randomstream instead of global random to prep for potentially visualizing noise and/or doing fractures on background threads

#rnx

diff --git a/Engine/Plugins/Experimental/ChaosEditor/Source/FractureEditor/Private/FractureToolBrick.cpp b/Engine/Plugins/Experimental/ChaosEditor/Source/FractureEditor/Private/FractureToolBrick.cpp
--- a/Engine/Plugins/Experimental/ChaosEditor/Source/FractureEditor/Private/FractureToolBrick.cpp
+++ b/Engine/Plugins/Experimental/ChaosEditor/Source/FractureEditor/Private/FractureToolBrick.cpp
@@ -417,7 +417,7 @@ int32 UFractureToolBrick::ExecuteFracture(const FFractureToolContext& FractureCo
 		// Proximity is invalidated.
 		ClearProximity(FractureContext.GetGeometryCollection().Get());
 
-		return CutMultipleWithPlanarCells(VoronoiPlanarCells, *FractureContext.GetGeometryCollection(), FractureContext.GetSelection(), 0, CollisionSettings->PointSpacing, FractureContext.GetTransform());
+		return CutMultipleWithPlanarCells(VoronoiPlanarCells, *FractureContext.GetGeometryCollection(), FractureContext.GetSelection(), 0, CollisionSettings->PointSpacing, FractureContext.GetSeed(), FractureContext.GetTransform());
 	}
 
 	return INDEX_NONE;
diff --git a/Engine/Plugins/Experimental/ChaosEditor/Source/FractureEditor/Private/FractureToolCutter.cpp b/Engine/Plugins/Experimental/ChaosEditor/Source/FractureEditor/Private/FractureToolCutter.cpp
--- a/Engine/Plugins/Experimental/ChaosEditor/Source/FractureEditor/Private/FractureToolCutter.cpp
+++ b/Engine/Plugins/Experimental/ChaosEditor/Source/FractureEditor/Private/FractureToolCutter.cpp
@@ -360,7 +360,7 @@ int32 UFractureToolVoronoiCutterBase::ExecuteFracture(const FFractureToolContext
 		// Proximity is invalidated.
 		ClearProximity(FractureContext.GetGeometryCollection().Get());
 
-		return CutMultipleWithPlanarCells(VoronoiPlanarCells, *(FractureContext.GetGeometryCollection()), FractureContext.GetSelection(), CutterSettings->Grout, CollisionSettings->PointSpacing, FractureContext.GetTransform());
+		return CutMultipleWithPlanarCells(VoronoiPlanarCells, *(FractureContext.GetGeometryCollection()), FractureContext.GetSelection(), CutterSettings->Grout, CollisionSettings->PointSpacing, FractureContext.GetSeed(), FractureContext.GetTransform());
 	}
 
 	return INDEX_NONE;
diff --git a/Engine/Plugins/Experimental/ChaosEditor/Source/FractureEditor/Private/FractureToolPlaneCut.cpp b/Engine/Plugins/Experimental/ChaosEditor/Source/FractureEditor/Private/FractureToolPlaneCut.cpp
--- a/Engine/Plugins/Experimental/ChaosEditor/Source/FractureEditor/Private/FractureToolPlaneCut.cpp
+++ b/Engine/Plugins/Experimental/ChaosEditor/Source/FractureEditor/Private/FractureToolPlaneCut.cpp
@@ -197,7 +197,7 @@ int32 UFractureToolPlaneCut::ExecuteFracture(const FFractureToolContext& Fractur
 		// Proximity is invalidated.
 		ClearProximity(FractureContext.GetGeometryCollection().Get());
 
-		return CutMultipleWithMultiplePlanes(CuttingPlanes, InternalSurfaceMaterials, *FractureContext.GetGeometryCollection(), FractureContext.GetSelection(), CutterSettings->Grout, CollisionSettings->PointSpacing, FractureContext.GetTransform());
+		return CutMultipleWithMultiplePlanes(CuttingPlanes, InternalSurfaceMaterials, *FractureContext.GetGeometryCollection(), FractureContext.GetSelection(), CutterSettings->Grout, CollisionSettings->PointSpacing, FractureContext.GetSeed(), FractureContext.GetTransform());
 	}
 
 	return INDEX_NONE;
diff --git a/Engine/Plugins/Experimental/ChaosEditor/Source/FractureEditor/Private/FractureToolSlice.cpp b/Engine/Plugins/Experimental/ChaosEditor/Source/FractureEditor/Private/FractureToolSlice.cpp
--- a/Engine/Plugins/Experimental/ChaosEditor/Source/FractureEditor/Private/FractureToolSlice.cpp
+++ b/Engine/Plugins/Experimental/ChaosEditor/Source/FractureEditor/Private/FractureToolSlice.cpp
@@ -185,7 +185,7 @@ int32 UFractureToolSlice::ExecuteFracture(const FFractureToolContext& FractureCo
 		// Proximity is invalidated.
 		ClearProximity(FractureContext.GetGeometryCollection().Get());
 
-		return CutMultipleWithMultiplePlanes(CuttingPlanes, InternalSurfaceMaterials, *FractureContext.GetGeometryCollection(), FractureContext.GetSelection(), CutterSettings->Grout, CollisionSettings->PointSpacing, FractureContext.GetTransform());
+		return CutMultipleWithMultiplePlanes(CuttingPlanes, InternalSurfaceMaterials, *FractureContext.GetGeometryCollection(), FractureContext.GetSelection(), CutterSettings->Grout, CollisionSettings->PointSpacing, FractureContext.GetSeed(), FractureContext.GetTransform());
 	}
 
 	return INDEX_NONE;
diff --git a/Engine/Plugins/Experimental/PlanarCutPlugin/Source/PlanarCut/Private/GeometryMeshConversion.cpp b/Engine/Plugins/Experimental/PlanarCutPlugin/Source/PlanarCut/Private/GeometryMeshConversion.cpp
--- a/Engine/Plugins/Experimental/PlanarCutPlugin/Source/PlanarCut/Private/GeometryMeshConversion.cpp
+++ b/Engine/Plugins/Experimental/PlanarCutPlugin/Source/PlanarCut/Private/GeometryMeshConversion.cpp
@@ -1171,9 +1171,9 @@ void SetGeometryCollectionAttributes(FDynamicMesh3& Mesh, int32 NumUVLayers)
 }
 
 
-FCellMeshes::FCellMeshes(int32 NumUVLayersIn, const FPlanarCells& Cells, FAxisAlignedBox3d DomainBounds, double Grout, double ExtendDomain, bool bIncludeOutsideCell)
+FCellMeshes::FCellMeshes(int32 NumUVLayersIn, FRandomStream& RandomStream, const FPlanarCells& Cells, FAxisAlignedBox3d DomainBounds, double Grout, double ExtendDomain, bool bIncludeOutsideCell)
 {
-	Init(NumUVLayersIn, Cells, DomainBounds, Grout, ExtendDomain, bIncludeOutsideCell);
+	Init(NumUVLayersIn, RandomStream, Cells, DomainBounds, Grout, ExtendDomain, bIncludeOutsideCell);
 }
 
 
@@ -1327,10 +1327,10 @@ void FCellMeshes::ApplyNoise(FDynamicMesh3& Mesh, FVector3d Normal, const FNoise
 }
 
 
-void FCellMeshes::Init(int32 NumUVLayersIn, const FPlanarCells& Cells, FAxisAlignedBox3d DomainBounds, double Grout, double ExtendDomain, bool bIncludeOutsideCell)
+void FCellMeshes::Init(int32 NumUVLayersIn, FRandomStream& RandomStream, const FPlanarCells& Cells, FAxisAlignedBox3d DomainBounds, double Grout, double ExtendDomain, bool bIncludeOutsideCell)
 {
 	NumUVLayers = NumUVLayersIn;
-	InitEmpty();
+	InitEmpty(RandomStream);
 
 	float GlobalUVScale = Cells.InternalSurfaceMaterials.GlobalUVScale;
 	if (!ensure(GlobalUVScale > 0))
@@ -2218,6 +2218,7 @@ int32 FDynamicMeshCollection::CutWithMultiplePlanes(
 	const TArrayView<const FPlane>& Planes,
 	double Grout,
 	double CollisionSampleSpacing,
+	int32 RandomSeed,
 	FGeometryCollection* Collection,
 	FInternalSurfaceMaterials& InternalSurfaceMaterials,
 	bool bSetDefaultInternalMaterialsFromCollection
@@ -2243,6 +2244,8 @@ int32 FDynamicMeshCollection::CutWithMultiplePlanes(
 
 	int32 NumUVLayers = Collection->NumUVLayers();
 
+	FRandomStream RandomStream(RandomSeed);
+
 	if (bHasGrout)
 	{
 		// For multi-plane cuts with grout specifically, the easiest path seems to be:
@@ -2251,7 +2254,7 @@ int32 FDynamicMeshCollection::CutWithMultiplePlanes(
 		// 3. Use the generic CutWithCellMeshes path, where that grout mesh is both the inner and outside cell mesh
 		//    (Note the outside cell mesh is subtracted, not intersected)
 		//    (Note this relies on island splitting to separate all the pieces afterwards.)
-		FCellMeshes GroutCells(NumUVLayers);
+		FCellMeshes GroutCells(NumUVLayers, RandomStream);
 		GroutCells.SetNumCells(2);
 		FDynamicMesh3& GroutMesh = GroutCells.CellMeshes[0].AugMesh;
 		FDynamicMeshEditor GroutAppender(&GroutMesh);
@@ -2261,7 +2264,7 @@ int32 FDynamicMeshCollection::CutWithMultiplePlanes(
 			EnterProgressFrame(.5);
 			FPlanarCells PlaneCells(Planes[PlaneIdx]);
 			PlaneCells.InternalSurfaceMaterials = InternalSurfaceMaterials;
-			FCellMeshes PlaneGroutMesh(NumUVLayers);
+			FCellMeshes PlaneGroutMesh(NumUVLayers, RandomStream);
 			PlaneGroutMesh.MakeOnlyPlanarGroutCell(PlaneCells, Bounds, Grout);
 			GroutAppender.AppendMesh(&PlaneGroutMesh.CellMeshes[0].AugMesh, IndexMaps);
 		}
@@ -2367,7 +2370,7 @@ int32 FDynamicMeshCollection::CutWithMultiplePlanes(
 		FPlanarCells PlaneCells(Planes[PlaneIdx]);
 		PlaneCells.InternalSurfaceMaterials = InternalSurfaceMaterials;
 		double OnePercentExtend = Bounds.MaxDim() * .01;
-		FCellMeshes CellMeshes(NumUVLayers, PlaneCells, Bounds, 0, OnePercentExtend, false);
+		FCellMeshes CellMeshes(NumUVLayers, RandomStream, PlaneCells, Bounds, 0, OnePercentExtend, false);
 
 		// TODO: we could do these cuts in parallel (will takes some rework of the proximity and how results are added to the ToCut array)
 		for (int32 ToCutIdx = 0, ToCutNum = ToCut.Num(); ToCutIdx < ToCutNum; ToCutIdx++)
diff --git a/Engine/Plugins/Experimental/PlanarCutPlugin/Source/PlanarCut/Private/GeometryMeshConversion.h b/Engine/Plugins/Experimental/PlanarCutPlugin/Source/PlanarCut/Private/GeometryMeshConversion.h
--- a/Engine/Plugins/Experimental/PlanarCutPlugin/Source/PlanarCut/Private/GeometryMeshConversion.h
+++ b/Engine/Plugins/Experimental/PlanarCutPlugin/Source/PlanarCut/Private/GeometryMeshConversion.h
@@ -80,13 +80,14 @@ struct PLANARCUT_API FCellMeshes
 		}
 	}
 
-	FCellMeshes(int32 NumUVLayers)
+	FCellMeshes(int32 NumUVLayers, FRandomStream& RandomStream)
 	{
-		InitEmpty();
+		InitEmpty(RandomStream);
 	}
 
-	FCellMeshes(int32 NumUVLayers, const FPlanarCells& Cells, UE::Geometry::FAxisAlignedBox3d DomainBounds, double Grout, double ExtendDomain, bool bIncludeOutsideCell);
+	FCellMeshes(int32 NumUVLayers, FRandomStream& RandomStream, const FPlanarCells& Cells, UE::Geometry::FAxisAlignedBox3d DomainBounds, double Grout, double ExtendDomain, bool bIncludeOutsideCell);
 
+	// Note: RandomStream not required for this constructor because noise is not supported in this case
 	FCellMeshes(int32 NumUVLayers, UE::Geometry::FDynamicMesh3& SingleCutter, const FInternalSurfaceMaterials& Materials, TOptional<FTransform> Transform);
 
 	// Special function to just make the "grout" part of the planar mesh cells
@@ -121,15 +122,15 @@ struct PLANARCUT_API FCellMeshes
 		return MaterialID >= 0 ? -1 : -(MaterialID + 1);
 	}
 
-	void InitEmpty()
+	void InitEmpty(FRandomStream& RandomStream)
 	{
-		NoiseOffsetX = FMath::VRand() * 100;
-		NoiseOffsetY = FMath::VRand() * 100;
-		NoiseOffsetZ = FMath::VRand() * 100;
+		NoiseOffsetX = RandomStream.VRand() * 100;
+		NoiseOffsetY = RandomStream.VRand() * 100;
+		NoiseOffsetZ = RandomStream.VRand() * 100;
 		OutsideCellIndex = -1;
 	}
 
-	void Init(int32 NumUVLayersIn, const FPlanarCells& Cells, UE::Geometry::FAxisAlignedBox3d DomainBounds, double Grout, double ExtendDomain, bool bIncludeOutsideCell);
+	void Init(int32 NumUVLayersIn, FRandomStream& RandomStream, const FPlanarCells& Cells, UE::Geometry::FAxisAlignedBox3d DomainBounds, double Grout, double ExtendDomain, bool bIncludeOutsideCell);
 
 	void ApplyGeneralGrout(double Grout);
 
@@ -218,6 +219,7 @@ struct PLANARCUT_API FDynamicMeshCollection
 		const TArrayView<const FPlane>& Planes,
 		double Grout,
 		double CollisionSampleSpacing,
+		int32 RandomSeed,
 		FGeometryCollection* Collection,
 		FInternalSurfaceMaterials& InternalSurfaceMaterials,
 		bool bSetDefaultInternalMaterialsFromCollection
diff --git a/Engine/Plugins/Experimental/PlanarCutPlugin/Source/PlanarCut/Private/PlanarCut.cpp b/Engine/Plugins/Experimental/PlanarCutPlugin/Source/PlanarCut/Private/PlanarCut.cpp
--- a/Engine/Plugins/Experimental/PlanarCutPlugin/Source/PlanarCut/Private/PlanarCut.cpp
+++ b/Engine/Plugins/Experimental/PlanarCutPlugin/Source/PlanarCut/Private/PlanarCut.cpp
@@ -541,6 +541,7 @@ int32 CutWithPlanarCells(
 	int32 TransformIdx,
 	double Grout,
 	double CollisionSampleSpacing,
+	int32 RandomSeed,
 	const TOptional<FTransform>& TransformCollection,
 	bool bIncludeOutsideCellInOutput,
 	float CheckDistanceAcrossOutsideCellForProximity,
@@ -548,7 +549,7 @@ int32 CutWithPlanarCells(
 )
 {
 	TArray<int32> TransformIndices { TransformIdx };
-	return CutMultipleWithPlanarCells(Cells, Source, TransformIndices, Grout, CollisionSampleSpacing, TransformCollection, bIncludeOutsideCellInOutput, CheckDistanceAcrossOutsideCellForProximity, bSetDefaultInternalMaterialsFromCollection);
+	return CutMultipleWithPlanarCells(Cells, Source, TransformIndices, Grout, CollisionSampleSpacing, RandomSeed, TransformCollection, bIncludeOutsideCellInOutput, CheckDistanceAcrossOutsideCellForProximity, bSetDefaultInternalMaterialsFromCollection);
 }
 
 
@@ -559,6 +560,7 @@ int32 CutMultipleWithMultiplePlanes(
 	const TArrayView<const int32>& TransformIndices,
 	double Grout,
 	double CollisionSampleSpacing,
+	int32 RandomSeed,
 	const TOptional<FTransform>& TransformCollection,
 	bool bSetDefaultInternalMaterialsFromCollection
 )
@@ -582,7 +584,7 @@ int32 CutMultipleWithMultiplePlanes(
 	FDynamicMeshCollection MeshCollection(&Collection, TransformIndices, CollectionToWorld);
 
 	int32 NewGeomStartIdx = -1;
-	NewGeomStartIdx = MeshCollection.CutWithMultiplePlanes(Planes, Grout, CollisionSampleSpacing, &Collection, InternalSurfaceMaterials, bSetDefaultInternalMaterialsFromCollection);
+	NewGeomStartIdx = MeshCollection.CutWithMultiplePlanes(Planes, Grout, CollisionSampleSpacing, RandomSeed, &Collection, InternalSurfaceMaterials, bSetDefaultInternalMaterialsFromCollection);
 
 	Collection.ReindexMaterials();
 	return NewGeomStartIdx;
@@ -595,6 +597,7 @@ int32 CutMultipleWithPlanarCells(
 	const TArrayView<const int32>& TransformIndices,
 	double Grout,
 	double CollisionSampleSpacing,
+	int32 RandomSeed,
 	const TOptional<FTransform>& TransformCollection,
 	bool bIncludeOutsideCellInOutput,
 	float CheckDistanceAcrossOutsideCellForProximity,
@@ -616,7 +619,8 @@ int32 CutMultipleWithPlanarCells(
 
 	FDynamicMeshCollection MeshCollection(&Source, TransformIndices, CollectionToWorld);
 	double OnePercentExtend = MeshCollection.Bounds.MaxDim() * .01;
-	FCellMeshes CellMeshes(Source.NumUVLayers(), Cells, MeshCollection.Bounds, Grout, OnePercentExtend, bIncludeOutsideCellInOutput);
+	FRandomStream RandomStream(RandomSeed);
+	FCellMeshes CellMeshes(Source.NumUVLayers(), RandomStream, Cells, MeshCollection.Bounds, Grout, OnePercentExtend, bIncludeOutsideCellInOutput);
 
 	int32 NewGeomStartIdx = -1;
 
diff --git a/Engine/Plugins/Experimental/PlanarCutPlugin/Source/PlanarCut/Public/PlanarCut.h b/Engine/Plugins/Experimental/PlanarCutPlugin/Source/PlanarCut/Public/PlanarCut.h
--- a/Engine/Plugins/Experimental/PlanarCutPlugin/Source/PlanarCut/Public/PlanarCut.h
+++ b/Engine/Plugins/Experimental/PlanarCutPlugin/Source/PlanarCut/Public/PlanarCut.h
@@ -144,7 +144,8 @@ struct PLANARCUT_API FPlanarCells
  * @param Collection		The collection to be cut
  * @param TransformIdx		Which transform inside the collection to cut
  * @param Grout				Separation to leave between cutting cells
- * @param CollisionSampleSpacing	Target spacing between collision sample vertices	
+ * @param CollisionSampleSpacing	Target spacing between collision sample vertices
+ * @param RandomSeed				Seed to be used for random noise displacement
  * @param TransformCollection		Optional transform of the whole geometry collection; if unset, defaults to Identity
  * @param bIncludeOutsideCellInOutput	If true, geometry that was not inside any of the cells (e.g. was outside of the bounds of all cutting geometry) will still be included in the output; if false, it will be discarded.
  * @param CheckDistanceAcrossOutsideCellForProximity	If > 0, when a plane is neighboring the "outside" cell, instead of setting proximity to the outside cell, the algo will sample a point this far outside the cell in the normal direction of the plane to see if there is actually a non-outside cell there.  (Useful for bricks w/out mortar)
@@ -156,6 +157,7 @@ int32 PLANARCUT_API CutWithPlanarCells(
 	int32 TransformIdx,
 	double Grout,
 	double CollisionSampleSpacing,
+	int32 RandomSeed,
 	const TOptional<FTransform>& TransformCollection = TOptional<FTransform>(),
 	bool bIncludeOutsideCellInOutput = true,
 	float CheckDistanceAcrossOutsideCellForProximity = 0,
@@ -170,6 +172,7 @@ int32 PLANARCUT_API CutWithPlanarCells(
  * @param TransformIndices	Which transform groups inside the collection to cut
  * @param Grout				Separation to leave between cutting cells
  * @param CollisionSampleSpacing	Target spacing between collision sample vertices
+ * @param RandomSeed				Seed to be used for random noise displacement
  * @param TransformCollection		Optional transform of the whole geometry collection; if unset, defaults to Identity
  * @param bIncludeOutsideCellInOutput	If true, geometry that was not inside any of the cells (e.g. was outside of the bounds of all cutting geometry) will still be included in the output; if false, it will be discarded.
  * @param CheckDistanceAcrossOutsideCellForProximity	If > 0, when a plane is neighboring the "outside" cell, instead of setting proximity to the outside cell, the algo will sample a point this far outside the cell in the normal direction of the plane to see if there is actually a non-outside cell there.  (Useful for bricks w/out mortar)
@@ -181,6 +184,7 @@ int32 PLANARCUT_API CutMultipleWithPlanarCells(
 	const TArrayView<const int32>& TransformIndices,
 	double Grout,
 	double CollisionSampleSpacing,
+	int32 RandomSeed,
 	const TOptional<FTransform>& TransformCollection = TOptional<FTransform>(),
 	bool bIncludeOutsideCellInOutput = true,
 	float CheckDistanceAcrossOutsideCellForProximity = 0,  // TODO: < this param does nothing in the new mode; is only needed in special cases that aren't possible in the UI currently
@@ -196,6 +200,7 @@ int32 PLANARCUT_API CutMultipleWithPlanarCells(
  * @param TransformIndices	Which transform groups inside the collection to cut
  * @param Grout				Separation to leave between cutting cells
  * @param CollisionSampleSpacing	Target spacing between collision sample vertices
+ * @param RandomSeed				Seed to be used for random noise displacement
  * @param TransformCollection		Optional transform of the whole geometry collection; if unset, defaults to Identity
  * @return	index of first new geometry in the Output GeometryCollection, or -1 if no geometry was added
  */
@@ -206,6 +211,7 @@ int32 PLANARCUT_API CutMultipleWithMultiplePlanes(
 	const TArrayView<const int32>& TransformIndices,
 	double Grout,
 	double CollisionSampleSpacing,
+	int32 RandomSeed,
 	const TOptional<FTransform>& TransformCollection = TOptional<FTransform>(),
 	bool bSetDefaultInternalMaterialsFromCollection = true
 );
