fix handling of uv num channel limits (inclusive not exclusive bound)

diff --git a/Engine/Plugins/Experimental/PlanarCutPlugin/Source/PlanarCut/Private/GeometryMeshConversion.cpp b/Engine/Plugins/Experimental/PlanarCutPlugin/Source/PlanarCut/Private/GeometryMeshConversion.cpp
--- a/Engine/Plugins/Experimental/PlanarCutPlugin/Source/PlanarCut/Private/GeometryMeshConversion.cpp
+++ b/Engine/Plugins/Experimental/PlanarCutPlugin/Source/PlanarCut/Private/GeometryMeshConversion.cpp
@@ -199,7 +199,7 @@ namespace AugmentedDynamicMesh
 
 	void SetAllUV(FDynamicMesh3& Mesh, int VID, FVector2f UV, int NumUVLayers)
 	{
-		if (!ensure(NumUVLayers < MAX_NUM_UV_CHANNELS))
+		if (!ensure(NumUVLayers <= MAX_NUM_UV_CHANNELS))
 		{
 			NumUVLayers = MAX_NUM_UV_CHANNELS;
 		}
