make fracture tools not remove geometry on cluster nodes, as some users need it for e.g. targeting different platforms

diff --git a/Engine/Plugins/Experimental/ChaosEditor/Source/FractureEditor/Private/FractureToolEditing.cpp b/Engine/Plugins/Experimental/ChaosEditor/Source/FractureEditor/Private/FractureToolEditing.cpp
--- a/Engine/Plugins/Experimental/ChaosEditor/Source/FractureEditor/Private/FractureToolEditing.cpp
+++ b/Engine/Plugins/Experimental/ChaosEditor/Source/FractureEditor/Private/FractureToolEditing.cpp
@@ -274,16 +274,20 @@ void UFractureToolValidate::Execute(TWeakPtr<FFractureEditorModeToolkit> InToolk
 				{
 					bool bDirty = false;
 
-					// Ensure that clusters do not point to geometry
 					TManagedArray<int32>& TransformToGeometry = GeometryCollection->TransformToGeometryIndex;
-					const int32 ElementCount = TransformToGeometry.Num();
-					for (int32 Idx = 0; Idx < ElementCount; ++Idx)
+					constexpr bool bClustersCanHaveGeometry = true;
+					if (!bClustersCanHaveGeometry)
 					{
-						if (GeometryCollection->IsClustered(Idx) && TransformToGeometry[Idx] != INDEX_NONE)
+						// Optionally ensure that clusters do not point to geometry (currently disabled; keeping the geometry can be useful)
+						const int32 ElementCount = TransformToGeometry.Num();
+						for (int32 Idx = 0; Idx < ElementCount; ++Idx)
 						{
-							TransformToGeometry[Idx] = INDEX_NONE;
-							UE_LOG(LogFractureTool, Warning, TEXT("Removed geometry index from cluster %d."), Idx);
-							bDirty = true;
+							if (GeometryCollection->IsClustered(Idx) && TransformToGeometry[Idx] != INDEX_NONE)
+							{
+								TransformToGeometry[Idx] = INDEX_NONE;
+								UE_LOG(LogFractureTool, Warning, TEXT("Removed geometry index from cluster %d."), Idx);
+								bDirty = true;
+							}
 						}
 					}
 
diff --git a/Engine/Plugins/Experimental/PlanarCutPlugin/Source/PlanarCut/Private/GeometryMeshConversion.cpp b/Engine/Plugins/Experimental/PlanarCutPlugin/Source/PlanarCut/Private/GeometryMeshConversion.cpp
--- a/Engine/Plugins/Experimental/PlanarCutPlugin/Source/PlanarCut/Private/GeometryMeshConversion.cpp
+++ b/Engine/Plugins/Experimental/PlanarCutPlugin/Source/PlanarCut/Private/GeometryMeshConversion.cpp
@@ -2213,6 +2213,21 @@ void FDynamicMeshCollection::Init(const FGeometryCollection* Collection, const T
 }
 
 
+void FDynamicMeshCollection::SetGeometryVisibility(FGeometryCollection* Collection, const TArray<int32>& GeometryIndices, bool bVisible)
+{
+	TManagedArray<bool>& Visible = Collection->Visible;
+	for (int32 GeoIdx : GeometryIndices)
+	{
+		int32 FaceStart = Collection->FaceStart[GeoIdx];
+		int32 FaceEnd = FaceStart + Collection->FaceCount[GeoIdx];
+		for (int32 FaceIdx = FaceStart; FaceIdx < FaceEnd; FaceIdx++)
+		{
+			Visible[FaceIdx] = bVisible;
+		}
+	}
+}
+
+
 int32 FDynamicMeshCollection::CutWithMultiplePlanes(
 	const TArrayView<const FPlane>& Planes,
 	double Grout,
@@ -2566,15 +2581,22 @@ int32 FDynamicMeshCollection::CutWithMultiplePlanes(
 		}
 	}
 
-	// remove superfluous geometry
+	constexpr bool bRemoveOldGeometry = false; // if false, we just hide the geometry that we've replaced by fractured child geometry, rather than remove it
 	if (GeometryForRemoval.Num() > 0)
 	{
-		GeometryForRemoval.Sort();
-		FManagedArrayCollection::FProcessingParameters ProcessingParams;
+		if (bRemoveOldGeometry)
+		{
+			GeometryForRemoval.Sort();
+			FManagedArrayCollection::FProcessingParameters ProcessingParams;
 #if !UE_BUILD_DEBUG
-		ProcessingParams.bDoValidation = false;
+			ProcessingParams.bDoValidation = false;
 #endif
-		Collection->RemoveElements(FGeometryCollection::GeometryGroup, GeometryForRemoval, ProcessingParams);
+			Collection->RemoveElements(FGeometryCollection::GeometryGroup, GeometryForRemoval, ProcessingParams);
+		}
+		else
+		{
+			SetGeometryVisibility(Collection, GeometryForRemoval, false);
+		}
 	}
 
 	return FirstCreatedIndex - GeometryForRemoval.Num();
@@ -2757,15 +2779,23 @@ int32 FDynamicMeshCollection::CutWithCellMeshes(const FInternalSurfaceMaterials&
 		}
 	}
 
-	// remove superfluous geometry
+	// remove or hide superfluous geometry
+	constexpr bool bRemoveOldGeometry = false; // if false, we just hide the geometry that we've replaced by fractured child geometry, rather than remove it
 	if (GeometryForRemoval.Num() > 0)
 	{
-		GeometryForRemoval.Sort();
-		FManagedArrayCollection::FProcessingParameters ProcessingParams;
+		if (bRemoveOldGeometry)
+		{
+			GeometryForRemoval.Sort();
+			FManagedArrayCollection::FProcessingParameters ProcessingParams;
 #if !UE_BUILD_DEBUG
-		ProcessingParams.bDoValidation = false;
+			ProcessingParams.bDoValidation = false;
 #endif
-		Collection->RemoveElements(FGeometryCollection::GeometryGroup, GeometryForRemoval, ProcessingParams);
+			Collection->RemoveElements(FGeometryCollection::GeometryGroup, GeometryForRemoval, ProcessingParams);
+		}
+		else
+		{
+			SetGeometryVisibility(Collection, GeometryForRemoval, false);
+		}
 	}
 
 	return FirstIdx - GeometryForRemoval.Num();
diff --git a/Engine/Plugins/Experimental/PlanarCutPlugin/Source/PlanarCut/Private/GeometryMeshConversion.h b/Engine/Plugins/Experimental/PlanarCutPlugin/Source/PlanarCut/Private/GeometryMeshConversion.h
--- a/Engine/Plugins/Experimental/PlanarCutPlugin/Source/PlanarCut/Private/GeometryMeshConversion.h
+++ b/Engine/Plugins/Experimental/PlanarCutPlugin/Source/PlanarCut/Private/GeometryMeshConversion.h
@@ -265,6 +265,8 @@ struct PLANARCUT_API FDynamicMeshCollection
 
 private:
 
+	void SetGeometryVisibility(FGeometryCollection* Collection, const TArray<int32>& GeometryIndices, bool bVisible);
+
 	// Update an existing geometry in a collection w/ a new mesh (w/ the same number of faces and vertices!)
 	static bool UpdateCollection(const FTransform& FromCollection, UE::Geometry::FDynamicMesh3& Mesh, int32 GeometryIdx, FGeometryCollection& Output, int32 InternalMaterialID);
 
