fix crash due to incorrect "new first index" return from fracture mesh cutting functions

diff --git a/Engine/Plugins/Experimental/ChaosEditor/Source/FractureEditor/Private/FractureToolMeshCut.cpp b/Engine/Plugins/Experimental/ChaosEditor/Source/FractureEditor/Private/FractureToolMeshCut.cpp
--- a/Engine/Plugins/Experimental/ChaosEditor/Source/FractureEditor/Private/FractureToolMeshCut.cpp
+++ b/Engine/Plugins/Experimental/ChaosEditor/Source/FractureEditor/Private/FractureToolMeshCut.cpp
@@ -243,7 +243,7 @@ int32 UFractureToolMeshCut::ExecuteFracture(const FFractureToolContext& Fracture
 				{
 					FirstIndex = Index;
 				}
-				if (Index != -1)
+				if (Index > -1)
 				{
 					int32 TransformIdx = FractureContext.GetGeometryCollection()->TransformIndex[Index];
 					// after a successful cut, also consider any new bones added by the cut
diff --git a/Engine/Plugins/Experimental/PlanarCutPlugin/Source/PlanarCut/Private/GeometryMeshConversion.cpp b/Engine/Plugins/Experimental/PlanarCutPlugin/Source/PlanarCut/Private/GeometryMeshConversion.cpp
--- a/Engine/Plugins/Experimental/PlanarCutPlugin/Source/PlanarCut/Private/GeometryMeshConversion.cpp
+++ b/Engine/Plugins/Experimental/PlanarCutPlugin/Source/PlanarCut/Private/GeometryMeshConversion.cpp
@@ -2581,6 +2581,8 @@ int32 FDynamicMeshCollection::CutWithMultiplePlanes(
 		}
 	}
 
+	int32 NewFirstIndex = FirstCreatedIndex;
+
 	constexpr bool bRemoveOldGeometry = false; // if false, we just hide the geometry that we've replaced by fractured child geometry, rather than remove it
 	if (GeometryForRemoval.Num() > 0)
 	{
@@ -2592,6 +2594,7 @@ int32 FDynamicMeshCollection::CutWithMultiplePlanes(
 			ProcessingParams.bDoValidation = false;
 #endif
 			Collection->RemoveElements(FGeometryCollection::GeometryGroup, GeometryForRemoval, ProcessingParams);
+			NewFirstIndex -= GeometryForRemoval.Num();
 		}
 		else
 		{
@@ -2599,7 +2602,7 @@ int32 FDynamicMeshCollection::CutWithMultiplePlanes(
 		}
 	}
 
-	return FirstCreatedIndex - GeometryForRemoval.Num();
+	return NewFirstIndex;
 }
 
 
@@ -2779,6 +2782,8 @@ int32 FDynamicMeshCollection::CutWithCellMeshes(const FInternalSurfaceMaterials&
 		}
 	}
 
+	int32 NewFirstIdx = FirstIdx;
+
 	// remove or hide superfluous geometry
 	constexpr bool bRemoveOldGeometry = false; // if false, we just hide the geometry that we've replaced by fractured child geometry, rather than remove it
 	if (GeometryForRemoval.Num() > 0)
@@ -2791,6 +2796,7 @@ int32 FDynamicMeshCollection::CutWithCellMeshes(const FInternalSurfaceMaterials&
 			ProcessingParams.bDoValidation = false;
 #endif
 			Collection->RemoveElements(FGeometryCollection::GeometryGroup, GeometryForRemoval, ProcessingParams);
+			NewFirstIdx -= GeometryForRemoval.Num();
 		}
 		else
 		{
@@ -2798,7 +2804,7 @@ int32 FDynamicMeshCollection::CutWithCellMeshes(const FInternalSurfaceMaterials&
 		}
 	}
 
-	return FirstIdx - GeometryForRemoval.Num();
+	return NewFirstIdx;
 }
 
 
