fix crash in fracture hole filler and disable hole filler

diff --git a/Engine/Plugins/Experimental/PlanarCutPlugin/Source/PlanarCut/Private/GeometryMeshConversion.cpp b/Engine/Plugins/Experimental/PlanarCutPlugin/Source/PlanarCut/Private/GeometryMeshConversion.cpp
--- a/Engine/Plugins/Experimental/PlanarCutPlugin/Source/PlanarCut/Private/GeometryMeshConversion.cpp
+++ b/Engine/Plugins/Experimental/PlanarCutPlugin/Source/PlanarCut/Private/GeometryMeshConversion.cpp
@@ -1062,11 +1062,15 @@ namespace AugmentedDynamicMesh
 				int32 V1 = AllBoundaryVerts[Idx];
 
 				int32 EID01 = AllBoundaryEdges[LastIdx];
+				if (!Mesh.IsBoundaryEdge(EID01))
+				{
+					continue;
+				}
 				
 				int32 C0 = LastIdx - BIdx - 1;
 				int32 C1 = Idx - BIdx - 1;
 				FIndex2i EdgeV = Mesh.GetOrientedBoundaryEdgeV(EID01);
-				if (HoleComponents[CurBdry] == (int32)VertComponents.Find(EdgeV.A))
+				if (EdgeV.A >= 0 && HoleComponents[CurBdry] == (int32)VertComponents.Find(EdgeV.A))
 				{
 					UseVIDs[C0] = EdgeV.A;
 					UseVIDs[C1] = EdgeV.B;
@@ -1108,6 +1112,10 @@ namespace AugmentedDynamicMesh
 				int32 V0 = UseVIDs[Tri.A];
 				int32 V1 = UseVIDs[Tri.B];
 				int32 V2 = UseVIDs[Tri.C];
+				if (V0 == V1 || V1 == V2 || V0 == V2)
+				{
+					continue;
+				}
 				FVector3d P0 = Mesh.GetVertex(V0);
 				FVector3d P1 = Mesh.GetVertex(V1);
 				FVector3d P2 = Mesh.GetVertex(V2);
@@ -2437,7 +2445,11 @@ int32 FDynamicMeshCollection::CutWithMultiplePlanes(
 					bKeepResults = false;
 					break;
 				}
-				AugmentedDynamicMesh::FillHoles(BoolResults[CellIdx]->AugMesh, Boolean.AllNewEdges);
+				constexpr bool bFillHoles = false; // TODO: consider enabling hole filler optionally, or improving the hole filler and enabling it all the time
+				if (bFillHoles)
+				{
+					AugmentedDynamicMesh::FillHoles(BoolResults[CellIdx]->AugMesh, Boolean.AllNewEdges);
+				}
 			}
 
 			if (bKeepResults)
@@ -2647,7 +2659,11 @@ int32 FDynamicMeshCollection::CutWithCellMeshes(const FInternalSurfaceMaterials&
 						// note: failure cases won't be detected at all unless we weld edges,
 						//       which will require re-working how tangents are carried through
 					}
-					AugmentedDynamicMesh::FillHoles(AugBoolResult, Boolean.AllNewEdges);
+					constexpr bool bFillHoles = false; // TODO: consider enabling hole filler optionally, or improving the hole filler and enabling it all the time
+					if (bFillHoles)
+					{
+						AugmentedDynamicMesh::FillHoles(AugBoolResult, Boolean.AllNewEdges);
+					}
 				}
 			}, EParallelForFlags::None);
 
