fix fracture geomerge merging geometry into clusters

diff --git a/Engine/Plugins/Experimental/PlanarCutPlugin/Source/PlanarCut/Private/PlanarCut.cpp b/Engine/Plugins/Experimental/PlanarCutPlugin/Source/PlanarCut/Private/PlanarCut.cpp
--- a/Engine/Plugins/Experimental/PlanarCutPlugin/Source/PlanarCut/Private/PlanarCut.cpp
+++ b/Engine/Plugins/Experimental/PlanarCutPlugin/Source/PlanarCut/Private/PlanarCut.cpp
@@ -1044,7 +1044,7 @@ int32 MergeBones(
 	for (int32 TransformIdx : TransformIndices)
 	{
 		int32 GeomIdx = Collection.TransformToGeometryIndex[TransformIdx];
-		if (GeomIdx > -1)
+		if (GeomIdx > -1 && Collection.IsRigid(TransformIdx))
 		{
 			CanMerge.Add(GeomIdx);
 			GeomToVol.Add(GeomIdx, Volumes[TransformIdx]);
@@ -1053,7 +1053,7 @@ int32 MergeBones(
 	for (int32 TransformIdx : SmallTransformIndices)
 	{
 		int32 GeomIdx = Collection.TransformToGeometryIndex[TransformIdx];
-		if (GeomIdx > -1)
+		if (GeomIdx > -1 && Collection.IsRigid(TransformIdx))
 		{
 			TooSmalls.Add(GeomIdx);
 			GeomToVol.Add(GeomIdx, Volumes[TransformIdx]);
