remove proximity updating from fracture code, as there is a better proximity update that will run instead after

diff --git a/Engine/Plugins/Experimental/PlanarCutPlugin/Source/PlanarCut/Private/GeometryMeshConversion.cpp b/Engine/Plugins/Experimental/PlanarCutPlugin/Source/PlanarCut/Private/GeometryMeshConversion.cpp
--- a/Engine/Plugins/Experimental/PlanarCutPlugin/Source/PlanarCut/Private/GeometryMeshConversion.cpp
+++ b/Engine/Plugins/Experimental/PlanarCutPlugin/Source/PlanarCut/Private/GeometryMeshConversion.cpp
@@ -2309,83 +2309,12 @@ int32 FDynamicMeshCollection::CutWithMultiplePlanes(
 		return CutWithCellMeshes(InternalSurfaceMaterials, CellConnectivity, GroutCells, Collection, bSetDefaultInternalMaterialsFromCollection, CollisionSampleSpacing);
 	}
 
-	bool bHasProximity = Collection->HasAttribute("Proximity", FGeometryCollection::GeometryGroup);
 	TArray<TUniquePtr<FMeshData>> ToCut;
-	TArray<TUniquePtr<TPointHashGrid3d<int>>> VerticesHashes;
-	auto HashMeshVertices = [&VerticesHashes, &ToCut](int32 HashIdx)
-	{
-		FDynamicMesh3& Mesh = ToCut[HashIdx]->AugMesh;
-		if (HashIdx >= VerticesHashes.Num())
-		{
-			VerticesHashes.SetNum(HashIdx + 1);
-		}
-		if (VerticesHashes[HashIdx].IsValid())
-		{
-			return;
-		}
-		VerticesHashes[HashIdx] = MakeUnique<TPointHashGrid3d<int>>(FMathd::ZeroTolerance * 1000, -1);
-		TPointHashGrid3d<int>& Grid = *VerticesHashes[HashIdx].Get();
-		for (int VID : Mesh.VertexIndicesItr())
-		{
-			Grid.InsertPointUnsafe(VID, Mesh.GetVertex(VID));
-		}
-	};
-	auto ClearHash = [&VerticesHashes](int32 HashIdx)
-	{
-		if (HashIdx < VerticesHashes.Num())
-		{
-			VerticesHashes[HashIdx].Release();
-		}
-	};
-	auto IsNeighbor = [&ToCut, &VerticesHashes](int32 A, int32 B)
-	{
-		if (!ensure(A < ToCut.Num() && B < ToCut.Num() && A < VerticesHashes.Num() && B < VerticesHashes.Num()))
-		{
-			return false;
-		}
-		if (!ensure(VerticesHashes[A].IsValid() && VerticesHashes[B].IsValid()))
-		{
-			return false;
-		}
-		if (!ToCut[A]->GetCachedBounds().Intersects(ToCut[B]->GetCachedBounds()))
-		{
-			return false;
-		}
-		if (ToCut[A]->AugMesh.VertexCount() > ToCut[B]->AugMesh.VertexCount())
-		{
-			Swap(A, B);
-		}
-		FDynamicMesh3& RefMesh = ToCut[B]->AugMesh;
-		for (const FVector3d& V : ToCut[A]->AugMesh.VerticesItr())
-		{
-			TPair<int, double> Nearest = VerticesHashes[B]->FindNearestInRadius(V, FMathd::ZeroTolerance * 10, [&RefMesh, &V](int VID)
-				{
-					return DistanceSquared(RefMesh.GetVertex(VID), V);
-				});
-			if (Nearest.Key != -1)
-			{
-				return true;
-			}
-		}
-		return false;
-	};
 	// copy initial surfaces
 	for (FMeshData& MeshData : Meshes)
 	{
 		ToCut.Add(MakeUnique<FMeshData>(MeshData));
 	}
-	// track connections between meshes via their indices in the ToCut array
-	TMultiMap<int32, int32> Proximity;
-	auto ProxLink = [&Proximity](int32 A, int32 B)
-	{
-		Proximity.Add(A, B);
-		Proximity.Add(B, A);
-	};
-	auto ProxUnlink = [&Proximity](int32 A, int32 B)
-	{
-		Proximity.RemoveSingle(A, B);
-		Proximity.RemoveSingle(B, A);
-	};
 	for (int32 PlaneIdx = 0; PlaneIdx < Planes.Num(); PlaneIdx++)
 	{
 		EnterProgressFrame(1);
@@ -2394,7 +2323,7 @@ int32 FDynamicMeshCollection::CutWithMultiplePlanes(
 		double OnePercentExtend = Bounds.MaxDim() * .01;
 		FCellMeshes CellMeshes(NumUVLayers, RandomStream, PlaneCells, Bounds, 0, OnePercentExtend, false);
 
-		// TODO: we could do these cuts in parallel (will takes some rework of the proximity and how results are added to the ToCut array)
+		// TODO: we could do these cuts in parallel (will takes some rework of how results are added to the ToCut array)
 		for (int32 ToCutIdx = 0, ToCutNum = ToCut.Num(); ToCutIdx < ToCutNum; ToCutIdx++)
 		{
 			FMeshData& Surface = *ToCut[ToCutIdx];
@@ -2475,68 +2404,6 @@ int32 FDynamicMeshCollection::CutWithMultiplePlanes(
 					}
 				}
 
-				// update proximity for neighbors of the original piece
-				if (bHasProximity)
-				{
-					ClearHash(ToCutIdx);
-					TArray<int32> Nbrs;
-					Proximity.MultiFind(ToCutIdx, Nbrs);
-					if (Nbrs.Num() > 0)
-					{
-						for (int32 ChangedMeshIdx : ResultIndices)
-						{
-							HashMeshVertices(ChangedMeshIdx);
-						}
-
-						for (int32 Nbr : Nbrs)
-						{
-							ProxUnlink(ToCutIdx, Nbr);
-							HashMeshVertices(Nbr);
-							for (int32 Idx = 0; Idx < ResultIndices.Num(); Idx++)
-							{
-								int32 ResultIdx = ResultIndices[Idx];
-								int32 OldIdx = Nbr;
-								if (IsNeighbor(ResultIdx, OldIdx))
-								{
-									ProxLink(ResultIdx, OldIdx);
-								}
-							}
-						}
-					}
-
-					if (ResultIndices.Num() == 2)
-					{
-						// add the connection between the two new pieces
-						ProxLink(ResultIndices[0], ResultIndices[1]);
-					}
-					else
-					{
-						if (Nbrs.Num() == 0)
-						{
-							for (int32 ChangedMeshIdx : ResultIndices)
-							{
-								HashMeshVertices(ChangedMeshIdx);
-							}
-						}
-						// check for connections between all pieces
-						for (int FirstIdx = 0; FirstIdx + 1 < ResultIndices.Num(); FirstIdx++)
-						{
-							int32 FirstParent = ParentIndices[FirstIdx];
-							for (int SecondIdx = FirstIdx + 1; SecondIdx < ResultIndices.Num(); SecondIdx++)
-							{
-								if (FirstParent == ParentIndices[SecondIdx])
-								{
-									// these pieces split from the same mesh *because* they were disconnected, so the pieces cannot be neighbors
-									continue;
-								}
-								if (IsNeighbor(ResultIndices[FirstIdx], ResultIndices[SecondIdx]))
-								{
-									ProxLink(ResultIndices[FirstIdx], ResultIndices[SecondIdx]);
-								}
-							}
-						}
-					}
-				}
 			} // iteration over meshes to cut
 		} // iteration over cutting planes
 	}
@@ -2583,16 +2450,6 @@ int32 FDynamicMeshCollection::CutWithMultiplePlanes(
 		}
 	}
 
-	// create proximity sets on geometry collection and populate using ToCut's Proximity multimap and the array ToCutIdxToGeometryIdx
-	if (bHasProximity)
-	{
-		TManagedArray<TSet<int32>>& GCProximity = Collection->ModifyAttribute<TSet<int32>>("Proximity", FGeometryCollection::GeometryGroup);
-		for (TPair<int32, int32> Link : Proximity)
-		{
-			GCProximity[ToCutIdxToGeometryIdx[Link.Key]].Add(ToCutIdxToGeometryIdx[Link.Value]);
-		}
-	}
-
 	int32 NewFirstIndex = FirstCreatedIndex;
 
 	constexpr bool bRemoveOldGeometry = false; // if false, we just hide the geometry that we've replaced by fractured child geometry, rather than remove it
@@ -2623,7 +2480,6 @@ int32 FDynamicMeshCollection::CutWithCellMeshes(const FInternalSurfaceMaterials&
 	// TODO: should we do these cuts in parallel, and the appends sequentially below?
 	int32 FirstIdx = -1;
 	int BadCount = 0;
-	bool bHasProximity = Collection->HasAttribute("Proximity", FGeometryCollection::GeometryGroup);
 	TArray<int32> GeometryForRemoval;
 	TArray<FAxisAlignedBox3d> CellBounds; CellBounds.Reserve(CellMeshes.CellMeshes.Num());
 	for (int CellIdx = 0; CellIdx < CellMeshes.CellMeshes.Num(); CellIdx++)
@@ -2738,60 +2594,6 @@ int32 FDynamicMeshCollection::CutWithCellMeshes(const FInternalSurfaceMaterials&
 					}
 				}
 			}
-			if (bHasProximity)
-			{
-				TManagedArray<TSet<int32>>& Proximity = Collection->ModifyAttribute<TSet<int32>>("Proximity", FGeometryCollection::GeometryGroup);
-				TArray<TUniquePtr<TPointHashGrid3d<int>>> VertexHashes;
-				TArray<FAxisAlignedBox3d> MeshBounds;
-				auto MakeHash = [this, &VertexHashes, &MeshBounds, &BooleanResults](int GID)
-				{
-					if (GID >= VertexHashes.Num())
-					{
-						VertexHashes.SetNum(GID + 1);
-						MeshBounds.SetNum(GID + 1);
-					}
-					if (!VertexHashes[GID].IsValid())
-					{
-						VertexHashes[GID] = MakeUnique<TPointHashGrid3d<int>>(FMathd::ZeroTolerance * 1000, -1);
-						FillVertexHash(*BooleanResults[GID], *VertexHashes[GID]);
-						MeshBounds[GID] = BooleanResults[GID]->GetBounds(true);
-					}
-				};
-				for (int32 PlaneIdx : PlanesInOutput)
-				{
-					TPair<int32, int32> Cells = CellConnectivity[PlaneIdx];
-					int32 SecondCell = Cells.Value < 0 ? CellMeshes.OutsideCellIndex : Cells.Value;
-					if (SecondCell != -1)
-					{
-						TArray<int32, TInlineAllocator<4>> GeomA, GeomB;
-						CellToGeometry.MultiFind(Cells.Key, GeomA, false);
-						CellToGeometry.MultiFind(SecondCell, GeomB, false);
-						if (GeomA.Num() == 1 && GeomB.Num() == 1)
-						{
-							Proximity[GeomA[0]].Add(GeomB[0]);
-							Proximity[GeomB[0]].Add(GeomA[0]);
-						}
-						else if (GeomA.Num() >= 1 && GeomB.Num() >= 1) // at least one was split; need to re-check proximities
-						{
-							for (int GIDA : GeomA)
-							{
-								int MeshA = GeometryToResultMesh[GIDA];
-								MakeHash(MeshA);
-								for (int GIDB : GeomB)
-								{
-									int MeshB = GeometryToResultMesh[GIDB];
-									MakeHash(MeshB);
-									if (IsNeighboring(*BooleanResults[MeshA], *VertexHashes[MeshA], MeshBounds[MeshA], *BooleanResults[MeshB], *VertexHashes[MeshB], MeshBounds[MeshB]))
-									{
-										Proximity[GIDA].Add(GIDB);
-										Proximity[GIDB].Add(GIDA);
-									}
-								}
-							}
-						}
-					}
-				}
-			}
 			
 			// remove old geom
 			GeometryForRemoval.Add(GeometryIdx);
diff --git a/Engine/Plugins/Experimental/PlanarCutPlugin/Source/PlanarCut/Private/PlanarCut.cpp b/Engine/Plugins/Experimental/PlanarCutPlugin/Source/PlanarCut/Private/PlanarCut.cpp
--- a/Engine/Plugins/Experimental/PlanarCutPlugin/Source/PlanarCut/Private/PlanarCut.cpp
+++ b/Engine/Plugins/Experimental/PlanarCutPlugin/Source/PlanarCut/Private/PlanarCut.cpp
@@ -544,12 +544,11 @@ int32 CutWithPlanarCells(
 	int32 RandomSeed,
 	const TOptional<FTransform>& TransformCollection,
 	bool bIncludeOutsideCellInOutput,
-	float CheckDistanceAcrossOutsideCellForProximity,
 	bool bSetDefaultInternalMaterialsFromCollection
 )
 {
 	TArray<int32> TransformIndices { TransformIdx };
-	return CutMultipleWithPlanarCells(Cells, Source, TransformIndices, Grout, CollisionSampleSpacing, RandomSeed, TransformCollection, bIncludeOutsideCellInOutput, CheckDistanceAcrossOutsideCellForProximity, bSetDefaultInternalMaterialsFromCollection);
+	return CutMultipleWithPlanarCells(Cells, Source, TransformIndices, Grout, CollisionSampleSpacing, RandomSeed, TransformCollection, bIncludeOutsideCellInOutput, bSetDefaultInternalMaterialsFromCollection);
 }
 
 
@@ -573,12 +572,6 @@ int32 CutMultipleWithMultiplePlanes(
 		InternalSurfaceMaterials.SetUVScaleFromCollection(Collection);
 	}
 
-	if (!Collection.HasAttribute("Proximity", FGeometryCollection::GeometryGroup))
-	{
-		const FManagedArrayCollection::FConstructionParameters GeometryDependency(FGeometryCollection::GeometryGroup);
-		Collection.AddAttribute<TSet<int32>>("Proximity", FGeometryCollection::GeometryGroup, GeometryDependency);
-	}
-
 	FTransform CollectionToWorld = TransformCollection.Get(FTransform::Identity);
 
 	FDynamicMeshCollection MeshCollection(&Collection, TransformIndices, CollectionToWorld);
@@ -600,16 +593,9 @@ int32 CutMultipleWithPlanarCells(
 	int32 RandomSeed,
 	const TOptional<FTransform>& TransformCollection,
 	bool bIncludeOutsideCellInOutput,
-	float CheckDistanceAcrossOutsideCellForProximity,
 	bool bSetDefaultInternalMaterialsFromCollection
 )
 {
-	if (!Source.HasAttribute("Proximity", FGeometryCollection::GeometryGroup))
-	{
-		const FManagedArrayCollection::FConstructionParameters GeometryDependency(FGeometryCollection::GeometryGroup);
-		Source.AddAttribute<TSet<int32>>("Proximity", FGeometryCollection::GeometryGroup, GeometryDependency);
-	}
-
 	if (bSetDefaultInternalMaterialsFromCollection)
 	{
 		Cells.InternalSurfaceMaterials.SetUVScaleFromCollection(Source);
@@ -735,12 +721,6 @@ int32 CutWithMesh(
 		}
 	}
 
-	if (!Collection.HasAttribute("Proximity", FGeometryCollection::GeometryGroup))
-	{
-		const FManagedArrayCollection::FConstructionParameters GeometryDependency(FGeometryCollection::GeometryGroup);
-		Collection.AddAttribute<TSet<int32>>("Proximity", FGeometryCollection::GeometryGroup, GeometryDependency);
-	}
-
 	if (bSetDefaultInternalMaterialsFromCollection)
 	{
 		InternalSurfaceMaterials.SetUVScaleFromCollection(Collection);
diff --git a/Engine/Plugins/Experimental/PlanarCutPlugin/Source/PlanarCut/Public/PlanarCut.h b/Engine/Plugins/Experimental/PlanarCutPlugin/Source/PlanarCut/Public/PlanarCut.h
--- a/Engine/Plugins/Experimental/PlanarCutPlugin/Source/PlanarCut/Public/PlanarCut.h
+++ b/Engine/Plugins/Experimental/PlanarCutPlugin/Source/PlanarCut/Public/PlanarCut.h
@@ -148,7 +148,6 @@ struct PLANARCUT_API FPlanarCells
  * @param RandomSeed				Seed to be used for random noise displacement
  * @param TransformCollection		Optional transform of the whole geometry collection; if unset, defaults to Identity
  * @param bIncludeOutsideCellInOutput	If true, geometry that was not inside any of the cells (e.g. was outside of the bounds of all cutting geometry) will still be included in the output; if false, it will be discarded.
- * @param CheckDistanceAcrossOutsideCellForProximity	If > 0, when a plane is neighboring the "outside" cell, instead of setting proximity to the outside cell, the algo will sample a point this far outside the cell in the normal direction of the plane to see if there is actually a non-outside cell there.  (Useful for bricks w/out mortar)
  * @return	index of first new geometry in the Output GeometryCollection, or -1 if no geometry was added
  */
 int32 PLANARCUT_API CutWithPlanarCells(
@@ -160,7 +159,6 @@ int32 PLANARCUT_API CutWithPlanarCells(
 	int32 RandomSeed,
 	const TOptional<FTransform>& TransformCollection = TOptional<FTransform>(),
 	bool bIncludeOutsideCellInOutput = true,
-	float CheckDistanceAcrossOutsideCellForProximity = 0,
 	bool bSetDefaultInternalMaterialsFromCollection = true
 );
 
@@ -175,7 +173,6 @@ int32 PLANARCUT_API CutWithPlanarCells(
  * @param RandomSeed				Seed to be used for random noise displacement
  * @param TransformCollection		Optional transform of the whole geometry collection; if unset, defaults to Identity
  * @param bIncludeOutsideCellInOutput	If true, geometry that was not inside any of the cells (e.g. was outside of the bounds of all cutting geometry) will still be included in the output; if false, it will be discarded.
- * @param CheckDistanceAcrossOutsideCellForProximity	If > 0, when a plane is neighboring the "outside" cell, instead of setting proximity to the outside cell, the algo will sample a point this far outside the cell in the normal direction of the plane to see if there is actually a non-outside cell there.  (Useful for bricks w/out mortar)
  * @return	index of first new geometry in the Output GeometryCollection, or -1 if no geometry was added
  */
 int32 PLANARCUT_API CutMultipleWithPlanarCells(
@@ -187,7 +184,6 @@ int32 PLANARCUT_API CutMultipleWithPlanarCells(
 	int32 RandomSeed,
 	const TOptional<FTransform>& TransformCollection = TOptional<FTransform>(),
 	bool bIncludeOutsideCellInOutput = true,
-	float CheckDistanceAcrossOutsideCellForProximity = 0,  // TODO: < this param does nothing in the new mode; is only needed in special cases that aren't possible in the UI currently
 	bool bSetDefaultInternalMaterialsFromCollection = true
 );
 
