fix transforms on fracture geomerge tool

diff --git a/Engine/Plugins/Experimental/PlanarCutPlugin/Source/PlanarCut/Private/PlanarCut.cpp b/Engine/Plugins/Experimental/PlanarCutPlugin/Source/PlanarCut/Private/PlanarCut.cpp
--- a/Engine/Plugins/Experimental/PlanarCutPlugin/Source/PlanarCut/Private/PlanarCut.cpp
+++ b/Engine/Plugins/Experimental/PlanarCutPlugin/Source/PlanarCut/Private/PlanarCut.cpp
@@ -1167,7 +1167,6 @@ int32 MergeBones(
 	for (int32 UpMeshIdx = 0; UpMeshIdx < UpdateCollection.Meshes.Num(); UpMeshIdx++)
 	{
 		FMeshData& UpMeshData = UpdateCollection.Meshes[UpMeshIdx];
-		FTransformSRT3d UpMeshXF = (FTransformSRT3d)UpMeshData.FromCollection;
 		int32 UpGeoIdx = Collection.TransformToGeometryIndex[UpMeshData.TransformIndex];
 		FRemoveGroup& Group = RemoveGroups[GeomIdxToRemoveGroupIdx[UpGeoIdx]];
 		if (!ensure(Group.IsValid()))
@@ -1178,7 +1177,6 @@ int32 MergeBones(
 		for (int32 RmGeoIdx : Group.ToRemove)
 		{
 			FMeshData& RmMeshData = RemoveCollection.Meshes[GeoIdxToRmMeshIdx[RmGeoIdx]];
-			FTransformSRT3d RmMeshXF = (FTransformSRT3d)RmMeshData.FromCollection;
 			if (bUnionJoinedPieces)
 			{
 				FMeshBoolean Boolean(&UpMeshData.AugMesh, &RmMeshData.AugMesh, &UpMeshData.AugMesh, FMeshBoolean::EBooleanOp::Union);
@@ -1189,15 +1187,7 @@ int32 MergeBones(
 			else
 			{
 				FMeshIndexMappings IndexMaps_Unused;
-				MeshEditor.AppendMesh(&RmMeshData.AugMesh, IndexMaps_Unused,
-					[&UpMeshXF, &RmMeshXF](int32 Unused, const FVector3d& Pos)
-					{
-						return UpMeshXF.TransformPosition(RmMeshXF.InverseTransformPosition(Pos));
-					},
-					[&UpMeshXF, &RmMeshXF](int32 Unused, const FVector3d& Normal)
-					{
-						return UpMeshXF.TransformNormal(RmMeshXF.InverseTransformNormal(Normal));
-					});
+				MeshEditor.AppendMesh(&RmMeshData.AugMesh, IndexMaps_Unused);
 			}
 		}
 	}
