expand progresscancel usage to Brick and  Mesh fracture tools
+ refactor mesh cut to do mesh conversion separately once, rather than repeatedly per transform

diff --git a/Engine/Plugins/Experimental/ChaosEditor/Source/FractureEditor/Private/FractureToolBrick.cpp b/Engine/Plugins/Experimental/ChaosEditor/Source/FractureEditor/Private/FractureToolBrick.cpp
--- a/Engine/Plugins/Experimental/ChaosEditor/Source/FractureEditor/Private/FractureToolBrick.cpp
+++ b/Engine/Plugins/Experimental/ChaosEditor/Source/FractureEditor/Private/FractureToolBrick.cpp
@@ -12,6 +12,9 @@
 #include "GeometryCollection/GeometryCollectionObject.h"
 #include "PlanarCut.h"
 #include "FractureToolContext.h"
+#include "FractureToolBackgroundTask.h"
+
+using namespace UE::Fracture;
 
 #define LOCTEXT_NAMESPACE "FractureBrick"
 
@@ -373,6 +376,35 @@ void UFractureToolBrick::AddBoxEdges(const FVector& Min, const FVector& Max)
 	Edges.Emplace(MakeTuple(FVector(Min.X, Max.Y, Max.Z), Max));
 }
 
+class FCellsFractureOp : public FGeometryCollectionOperator
+{
+public:
+	FCellsFractureOp(const FGeometryCollection& SourceCollection) : FGeometryCollectionOperator(SourceCollection)
+	{}
+
+	virtual ~FCellsFractureOp() = default;
+
+	TArray<int> Selection;
+	FPlanarCells Cells;
+	float PointSpacing;
+	float Grout = 0;
+	int Seed;
+	FTransform Transform;
+
+	// TGenericDataOperator interface:
+	virtual void CalculateResult(FProgressCancel* Progress) override
+	{
+		if (Progress && Progress->Cancelled())
+		{
+			return;
+		}
+
+		ResultGeometryIndex = CutMultipleWithPlanarCells(Cells, *CollectionCopy, Selection, Grout, PointSpacing, Seed, Transform, true, true, Progress);
+
+		SetResult(MoveTemp(CollectionCopy));
+	}
+};
+
 int32 UFractureToolBrick::ExecuteFracture(const FFractureToolContext& FractureContext)
 {
 	if (FractureContext.IsValid())
@@ -405,19 +437,23 @@ int32 UFractureToolBrick::ExecuteFracture(const FFractureToolContext& FractureCo
 			BricksToCut.Add(BrickBox.TransformBy(Trans));
 		}
 
- 		FPlanarCells VoronoiPlanarCells = FPlanarCells(BricksToCut);
-
-		FNoiseSettings NoiseSettings;
+		TUniquePtr<FCellsFractureOp> BrickOp = MakeUnique<FCellsFractureOp>(*(FractureContext.GetGeometryCollection()));
+		BrickOp->Selection = FractureContext.GetSelection();
+		BrickOp->Grout = 0; // CutterSettings->Grout; // Note: Grout is currently baked directly into the brick cells above
+		BrickOp->PointSpacing = CollisionSettings->GetPointSpacing();
+		BrickOp->Cells = FPlanarCells(BricksToCut);
 		if (CutterSettings->Amplitude > 0.0f)
 		{
-			CutterSettings->TransferNoiseSettings(NoiseSettings);
-			VoronoiPlanarCells.InternalSurfaceMaterials.NoiseSettings = NoiseSettings;
+			FNoiseSettings Settings;
+			CutterSettings->TransferNoiseSettings(Settings);
+			BrickOp->Cells.InternalSurfaceMaterials.NoiseSettings = Settings;
 		}
+		BrickOp->Seed = FractureContext.GetSeed();
+		BrickOp->Transform = FractureContext.GetTransform();
 
-		// Proximity is invalidated.
-		ClearProximity(FractureContext.GetGeometryCollection().Get());
-
-		return CutMultipleWithPlanarCells(VoronoiPlanarCells, *FractureContext.GetGeometryCollection(), FractureContext.GetSelection(), 0, CollisionSettings->GetPointSpacing(), FractureContext.GetSeed(), FractureContext.GetTransform());
+		int Result = RunCancellableGeometryCollectionOp<FCellsFractureOp>(*(FractureContext.GetGeometryCollection()),
+			MoveTemp(BrickOp), LOCTEXT("ComputingBrickFractureMessage", "Computing Brick Fracture"));
+		return Result;
 	}
 
 	return INDEX_NONE;
diff --git a/Engine/Plugins/Experimental/ChaosEditor/Source/FractureEditor/Private/FractureToolCutter.cpp b/Engine/Plugins/Experimental/ChaosEditor/Source/FractureEditor/Private/FractureToolCutter.cpp
--- a/Engine/Plugins/Experimental/ChaosEditor/Source/FractureEditor/Private/FractureToolCutter.cpp
+++ b/Engine/Plugins/Experimental/ChaosEditor/Source/FractureEditor/Private/FractureToolCutter.cpp
@@ -394,8 +394,7 @@ public:
 		FProgressCancel::FProgressScope FractureMeshProgress =
 			FProgressCancel::CreateScopeTo(Progress, 1, LOCTEXT("FractureMeshMessage", "Fracturing Mesh"));
 
-		// TODO: Pass the ProgressCancel down so the computation can be tracked and cancelled
-		ResultGeometryIndex = CutMultipleWithPlanarCells(VoronoiPlanarCells, *CollectionCopy, Selection, Grout, PointSpacing, Seed, Transform);
+		ResultGeometryIndex = CutMultipleWithPlanarCells(VoronoiPlanarCells, *CollectionCopy, Selection, Grout, PointSpacing, Seed, Transform, true, true, Progress);
 		
 		SetResult(MoveTemp(CollectionCopy));
 	}
diff --git a/Engine/Plugins/Experimental/ChaosEditor/Source/FractureEditor/Private/FractureToolMeshCut.cpp b/Engine/Plugins/Experimental/ChaosEditor/Source/FractureEditor/Private/FractureToolMeshCut.cpp
--- a/Engine/Plugins/Experimental/ChaosEditor/Source/FractureEditor/Private/FractureToolMeshCut.cpp
+++ b/Engine/Plugins/Experimental/ChaosEditor/Source/FractureEditor/Private/FractureToolMeshCut.cpp
@@ -14,6 +14,9 @@
 
 #include "Engine/StaticMesh.h"
 #include "Components/StaticMeshComponent.h"
+#include "FractureToolBackgroundTask.h"
+
+using namespace UE::Fracture;
 
 #define LOCTEXT_NAMESPACE "FractureMesh"
 
@@ -196,6 +199,68 @@ void UFractureToolMeshCut::FractureContextChanged()
 	}
 }
 
+class FMeshFractureOp : public FGeometryCollectionOperator
+{
+public:
+	FMeshFractureOp(const FGeometryCollection& SourceCollection) : FGeometryCollectionOperator(SourceCollection)
+	{}
+
+	virtual ~FMeshFractureOp() = default;
+
+	TArray<int> Selection;
+	TArray<FTransform> MeshTransforms;
+	float PointSpacing;
+	int Seed;
+	FTransform Transform;
+	UE::Geometry::FDynamicMesh3 CuttingMesh;
+
+	// TGenericDataOperator interface:
+	virtual void CalculateResult(FProgressCancel* Progress) override
+	{
+		// Note: Noise not currently supported
+		FInternalSurfaceMaterials InternalSurfaceMaterials;
+
+		ResultGeometryIndex = -1;
+		const float ProgressFrac = 1.0 / MeshTransforms.Num();
+		
+		TArray<int32> BonesToCut = Selection;
+		for (const FTransform& ScatterTransform : MeshTransforms)
+		{
+			if (Progress && Progress->Cancelled())
+			{
+				return;
+			}
+
+			FProgressCancel::FProgressScope ProgressScope(Progress, ProgressFrac);
+			int32 Index = CutWithMesh(CuttingMesh, ScatterTransform, InternalSurfaceMaterials, *CollectionCopy, BonesToCut, PointSpacing, Transform);
+			int32 NewLen = Algo::RemoveIf(BonesToCut, [&](int32 Bone)
+				{
+					return !CollectionCopy->IsVisible(Bone); // remove already-fractured pieces from the to-cut list
+				});
+			BonesToCut.SetNum(NewLen);
+			if (ResultGeometryIndex == -1)
+			{
+				ResultGeometryIndex = Index;
+			}
+			if (Index > -1)
+			{
+				int32 TransformIdx = CollectionCopy->TransformIndex[Index];
+				// after a successful cut, also consider any new bones added by the cut
+				for (int32 NewBoneIdx = TransformIdx; NewBoneIdx < CollectionCopy->NumElements(FGeometryCollection::TransformGroup); NewBoneIdx++)
+				{
+					BonesToCut.Add(NewBoneIdx);
+				}
+			}
+		}
+
+		if (Progress && Progress->Cancelled())
+		{
+			return;
+		}
+		SetResult(MoveTemp(CollectionCopy));
+	}
+};
+
 int32 UFractureToolMeshCut::ExecuteFracture(const FFractureToolContext& FractureContext)
 {
 	if (FractureContext.IsValid())
@@ -207,63 +272,25 @@ int32 UFractureToolMeshCut::ExecuteFracture(const FFractureToolContext& Fracture
 			return INDEX_NONE;
 		}
 
-		FMeshDescription* MeshDescription = LocalCutSettings->CuttingActor->GetStaticMeshComponent()->GetStaticMesh()->GetMeshDescription(0);
-		FTransform ActorTransform(LocalCutSettings->CuttingActor->GetTransform());
-		FInternalSurfaceMaterials InternalSurfaceMaterials;
-
-		// Proximity is invalidated.
-		ClearProximity(FractureContext.GetGeometryCollection().Get());
-
-		// (Note: noise and grout not currently supported)
+		TUniquePtr<FMeshFractureOp> MeshCutOp = MakeUnique<FMeshFractureOp>(*(FractureContext.GetGeometryCollection()));
+		MeshCutOp->Selection = FractureContext.GetSelection();
+		MeshCutOp->PointSpacing = CollisionSettings->GetPointSpacing();
+		MeshCutOp->Seed = FractureContext.GetSeed();
+		MeshCutOp->Transform = FractureContext.GetTransform();
 		if (LocalCutSettings->CutDistribution == EMeshCutDistribution::SingleCut)
 		{
-			return CutWithMesh(MeshDescription, ActorTransform, InternalSurfaceMaterials, *FractureContext.GetGeometryCollection(), FractureContext.GetSelection(), CollisionSettings->GetPointSpacing(), FractureContext.GetTransform());
+			MeshCutOp->MeshTransforms.Add(LocalCutSettings->CuttingActor->GetTransform());
 		}
 		else
 		{
-			TArray<FTransform> MeshTransforms;
-			GenerateMeshTransforms(FractureContext, MeshTransforms);
-			int32 FirstIndex = -1;
-			// Create progress indicator dialog
-			static const FText SlowTaskText = LOCTEXT("CutWithScatteredMeshes", "Cutting geometry collection with mesh ...");
-
-			FScopedSlowTask SlowTask(MeshTransforms.Num(), SlowTaskText);
-			SlowTask.MakeDialog();
-
-			// Declare progress shortcut lambdas
-			auto EnterProgressFrame = [&SlowTask](float Progress)
-			{
-				SlowTask.EnterProgressFrame(Progress);
-			};
-
-			// TODO: support passing multiple transforms to CutWithMesh, and doing the loop inside that function *after* the mesh conversions are done
-			// (As is, we're re-converting the cutting mesh for each transform!)
-			TArray<int32> BonesToCut = FractureContext.GetSelection();
-			for (const FTransform& ScatterTransform : MeshTransforms)
-			{
-				EnterProgressFrame(1);
-				int32 Index = CutWithMesh(MeshDescription, ScatterTransform, InternalSurfaceMaterials, *FractureContext.GetGeometryCollection(), BonesToCut, CollisionSettings->GetPointSpacing(), FractureContext.GetTransform());
-				int32 NewLen = Algo::RemoveIf(BonesToCut, [&FractureContext](int32 Bone)
-					{
-						return !FractureContext.GetGeometryCollection()->IsVisible(Bone); // remove already-fractured pieces from the to-cut list
-					});
-				BonesToCut.SetNum(NewLen);
-				if (FirstIndex == -1)
-				{
-					FirstIndex = Index;
-				}
-				if (Index > -1)
-				{
-					int32 TransformIdx = FractureContext.GetGeometryCollection()->TransformIndex[Index];
-					// after a successful cut, also consider any new bones added by the cut
-					for (int32 NewBoneIdx = TransformIdx; NewBoneIdx < FractureContext.GetGeometryCollection()->NumElements(FGeometryCollection::TransformGroup); NewBoneIdx++)
-					{
-						BonesToCut.Add(NewBoneIdx);
-					}
-				}
-			}
-			return FirstIndex;
+			GenerateMeshTransforms(FractureContext, MeshCutOp->MeshTransforms);
 		}
+		FMeshDescription* MeshDescription = LocalCutSettings->CuttingActor->GetStaticMeshComponent()->GetStaticMesh()->GetMeshDescription(0);
+		MeshCutOp->CuttingMesh = ConvertMeshDescriptionToCuttingDynamicMesh(MeshDescription, FractureContext.GetGeometryCollection()->NumUVLayers());
+
+		int Result = RunCancellableGeometryCollectionOp<FMeshFractureOp>(*(FractureContext.GetGeometryCollection()),
+			MoveTemp(MeshCutOp), LOCTEXT("ComputingVoronoiFractureMessage", "Computing Mesh Fracture"));
+		return Result;
 	}
 
 	return INDEX_NONE;
diff --git a/Engine/Plugins/Experimental/PlanarCutPlugin/Source/PlanarCut/Private/GeometryMeshConversion.cpp b/Engine/Plugins/Experimental/PlanarCutPlugin/Source/PlanarCut/Private/GeometryMeshConversion.cpp
--- a/Engine/Plugins/Experimental/PlanarCutPlugin/Source/PlanarCut/Private/GeometryMeshConversion.cpp
+++ b/Engine/Plugins/Experimental/PlanarCutPlugin/Source/PlanarCut/Private/GeometryMeshConversion.cpp
@@ -1185,23 +1185,25 @@ FCellMeshes::FCellMeshes(int32 NumUVLayersIn, FRandomStream& RandomStream, const
 }
 
 
-FCellMeshes::FCellMeshes(int32 NumUVLayersIn, FDynamicMesh3& SingleCutter, const FInternalSurfaceMaterials& Materials, TOptional<FTransform> Transform)
+FCellMeshes::FCellMeshes(int32 NumUVLayersIn, const FDynamicMesh3& SingleCutter, const FInternalSurfaceMaterials& Materials, TOptional<FTransform> Transform)
 {
 	NumUVLayers = NumUVLayersIn;
 	SetNumCells(2);
 
+	CellMeshes[0].AugMesh = SingleCutter;
+
 	if (Transform.IsSet())
 	{
-		MeshTransforms::ApplyTransform(SingleCutter, FTransformSRT3d(Transform.GetValue()));
+		MeshTransforms::ApplyTransform(CellMeshes[0].AugMesh, FTransformSRT3d(Transform.GetValue()));
 	}
 
 	// Mesh should already be augmented
-	if (!ensure(AugmentedDynamicMesh::IsAugmented(SingleCutter)))
+	if (!ensure(AugmentedDynamicMesh::IsAugmented(CellMeshes[0].AugMesh)))
 	{
-		AugmentedDynamicMesh::Augment(SingleCutter, NumUVLayers);
+		AugmentedDynamicMesh::Augment(CellMeshes[0].AugMesh, NumUVLayers);
 	}
 
-	CellMeshes[0].AugMesh = SingleCutter;
+	
 
 	// first mesh is the same as the second mesh, but will be subtracted b/c it's the "outside cell"
 	// TODO: special case this logic so we don't have to hold two copies of the exact same mesh!
diff --git a/Engine/Plugins/Experimental/PlanarCutPlugin/Source/PlanarCut/Private/GeometryMeshConversion.h b/Engine/Plugins/Experimental/PlanarCutPlugin/Source/PlanarCut/Private/GeometryMeshConversion.h
--- a/Engine/Plugins/Experimental/PlanarCutPlugin/Source/PlanarCut/Private/GeometryMeshConversion.h
+++ b/Engine/Plugins/Experimental/PlanarCutPlugin/Source/PlanarCut/Private/GeometryMeshConversion.h
@@ -89,7 +89,7 @@ struct PLANARCUT_API FCellMeshes
 	FCellMeshes(int32 NumUVLayers, FRandomStream& RandomStream, const FPlanarCells& Cells, UE::Geometry::FAxisAlignedBox3d DomainBounds, double Grout, double ExtendDomain, bool bIncludeOutsideCell);
 
 	// Note: RandomStream not required for this constructor because noise is not supported in this case
-	FCellMeshes(int32 NumUVLayers, UE::Geometry::FDynamicMesh3& SingleCutter, const FInternalSurfaceMaterials& Materials, TOptional<FTransform> Transform);
+	FCellMeshes(int32 NumUVLayers, const UE::Geometry::FDynamicMesh3& SingleCutter, const FInternalSurfaceMaterials& Materials, TOptional<FTransform> Transform);
 
 	// Special function to just make the "grout" part of the planar mesh cells
 	// Used to make the multi-plane cuts with grout easier to implement
diff --git a/Engine/Plugins/Experimental/PlanarCutPlugin/Source/PlanarCut/Private/PlanarCut.cpp b/Engine/Plugins/Experimental/PlanarCutPlugin/Source/PlanarCut/Private/PlanarCut.cpp
--- a/Engine/Plugins/Experimental/PlanarCutPlugin/Source/PlanarCut/Private/PlanarCut.cpp
+++ b/Engine/Plugins/Experimental/PlanarCutPlugin/Source/PlanarCut/Private/PlanarCut.cpp
@@ -52,6 +52,8 @@ using namespace UE::Geometry;
 
 using namespace UE::PlanarCut;
 
+#define LOCTEXT_NAMESPACE "PlanarCut"
+
 // logic from FMeshUtility::GenerateGeometryCollectionFromBlastChunk, sets material IDs based on construction pattern that external materials have even IDs and are matched to internal materials at InternalID = ExternalID+1
 int32 FInternalSurfaceMaterials::GetDefaultMaterialIDForGeometry(const FGeometryCollection& Collection, int32 GeometryIdx) const
 {
@@ -544,11 +546,13 @@ int32 CutWithPlanarCells(
 	int32 RandomSeed,
 	const TOptional<FTransform>& TransformCollection,
 	bool bIncludeOutsideCellInOutput,
-	bool bSetDefaultInternalMaterialsFromCollection
+	bool bSetDefaultInternalMaterialsFromCollection,
+	FProgressCancel* Progress
+
 )
 {
 	TArray<int32> TransformIndices { TransformIdx };
-	return CutMultipleWithPlanarCells(Cells, Source, TransformIndices, Grout, CollisionSampleSpacing, RandomSeed, TransformCollection, bIncludeOutsideCellInOutput, bSetDefaultInternalMaterialsFromCollection);
+	return CutMultipleWithPlanarCells(Cells, Source, TransformIndices, Grout, CollisionSampleSpacing, RandomSeed, TransformCollection, bIncludeOutsideCellInOutput, bSetDefaultInternalMaterialsFromCollection, Progress);
 }
 
 
@@ -561,9 +565,11 @@ int32 CutMultipleWithMultiplePlanes(
 	double CollisionSampleSpacing,
 	int32 RandomSeed,
 	const TOptional<FTransform>& TransformCollection,
-	bool bSetDefaultInternalMaterialsFromCollection
+	bool bSetDefaultInternalMaterialsFromCollection,
+	FProgressCancel* Progress
 )
 {
+	FProgressCancel::FProgressScope PrepareScope = FProgressCancel::CreateScopeTo(Progress, .1, LOCTEXT("CutWithMultiplePlanesInit", "Preparing to cut with planes"));
 	int32 OrigNumGeom = Collection.FaceCount.Num();
 	int32 CurNumGeom = OrigNumGeom;
 
@@ -576,10 +582,27 @@ int32 CutMultipleWithMultiplePlanes(
 
 	FDynamicMeshCollection MeshCollection(&Collection, TransformIndices, CollectionToWorld);
 
+	PrepareScope.Done();
+	if (Progress && Progress->Cancelled())
+	{
+		return -1;
+	}
+	FProgressCancel::FProgressScope CutScope = FProgressCancel::CreateScopeTo(Progress, .99);
+
 	int32 NewGeomStartIdx = -1;
 	NewGeomStartIdx = MeshCollection.CutWithMultiplePlanes(Planes, Grout, CollisionSampleSpacing, RandomSeed, &Collection, InternalSurfaceMaterials, bSetDefaultInternalMaterialsFromCollection);
 
+	CutScope.Done();
+	if (Progress && Progress->Cancelled())
+	{
+		return -1;
+	}
+	FProgressCancel::FProgressScope ReindexScope = FProgressCancel::CreateScopeTo(Progress, 1);
+
 	Collection.ReindexMaterials();
+
+	ReindexScope.Done();
+
 	return NewGeomStartIdx;
 }
 
@@ -593,9 +616,11 @@ int32 CutMultipleWithPlanarCells(
 	int32 RandomSeed,
 	const TOptional<FTransform>& TransformCollection,
 	bool bIncludeOutsideCellInOutput,
-	bool bSetDefaultInternalMaterialsFromCollection
+	bool bSetDefaultInternalMaterialsFromCollection,
+	FProgressCancel* Progress
 )
 {
+	FProgressCancel::FProgressScope CreateMeshCollectionScope = FProgressCancel::CreateScopeTo(Progress, .1);
 	if (bSetDefaultInternalMaterialsFromCollection)
 	{
 		Cells.InternalSurfaceMaterials.SetUVScaleFromCollection(Source);
@@ -604,31 +629,42 @@ int32 CutMultipleWithPlanarCells(
 	FTransform CollectionToWorld = TransformCollection.Get(FTransform::Identity);
 
 	FDynamicMeshCollection MeshCollection(&Source, TransformIndices, CollectionToWorld);
+	CreateMeshCollectionScope.Done();
+
+	if (Progress && Progress->Cancelled())
+	{
+		return -1;
+	}
+
+	FProgressCancel::FProgressScope CellMeshScope = FProgressCancel::CreateScopeTo(Progress, .1);
 	double OnePercentExtend = MeshCollection.Bounds.MaxDim() * .01;
 	FRandomStream RandomStream(RandomSeed);
 	FCellMeshes CellMeshes(Source.NumUVLayers(), RandomStream, Cells, MeshCollection.Bounds, Grout, OnePercentExtend, bIncludeOutsideCellInOutput);
+	CellMeshScope.Done();
 
-	int32 NewGeomStartIdx = -1;
+	if (Progress && Progress->Cancelled())
+	{
+		return -1;
+	}
 
+	FProgressCancel::FProgressScope CutScope = FProgressCancel::CreateScopeTo(Progress, .99);
+	int32 NewGeomStartIdx = -1;
 	NewGeomStartIdx = MeshCollection.CutWithCellMeshes(Cells.InternalSurfaceMaterials, Cells.PlaneCells, CellMeshes, &Source, bSetDefaultInternalMaterialsFromCollection, CollisionSampleSpacing);
+	CutScope.Done();
+
+	if (Progress && Progress->Cancelled())
+	{
+		return -1;
+	}
 
+	FProgressCancel::FProgressScope ReindexScope = FProgressCancel::CreateScopeTo(Progress, 1);
 	Source.ReindexMaterials();
+	ReindexScope.Done();
 	return NewGeomStartIdx;
 }
 
-int32 CutWithMesh(
-	FMeshDescription* CuttingMesh,
-	FTransform CuttingMeshTransform,
-	FInternalSurfaceMaterials& InternalSurfaceMaterials,
-	FGeometryCollection& Collection,
-	const TArrayView<const int32>& TransformIndices,
-	double CollisionSampleSpacing,
-	const TOptional<FTransform>& TransformCollection,
-	bool bSetDefaultInternalMaterialsFromCollection
-)
+FDynamicMesh3 ConvertMeshDescriptionToCuttingDynamicMesh(const FMeshDescription* CuttingMesh, int32 NumUVLayers, FProgressCancel* Progress)
 {
-	int32 NewGeomStartIdx = -1;
-
 	// populate the BaseMesh with a conversion of the input mesh.
 	FMeshDescriptionToDynamicMesh Converter;
 	FDynamicMesh3 FullMesh; // full-featured conversion of the source mesh
@@ -648,10 +684,19 @@ int32 CutWithMesh(
 		Tangents.CopyToOverlays(FullMesh);
 	}
 
+	if (Progress && Progress->Cancelled())
+	{
+		return FDynamicMesh3(); // return empty mesh on cancel
+	}
+
 	FDynamicMesh3 DynamicCuttingMesh; // version of mesh that is split apart at seams to be compatible w/ geometry collection, with corresponding attributes set
-	int32 NumUVLayers = Collection.NumUVLayers();
 	SetGeometryCollectionAttributes(DynamicCuttingMesh, NumUVLayers);
 
+	if (Progress && Progress->Cancelled())
+	{
+		return FDynamicMesh3(); // return empty mesh on cancel
+	}
+
 	// Note: This conversion will likely go away, b/c I plan to switch over to doing the boolean operations on the fuller rep, but the code can be adapted
 	//		 to the dynamic mesh -> geometry collection conversion phase, as this same splitting will then need to happen there.
 	if (ensure(FullMesh.HasAttributes() && FullMesh.Attributes()->NumUVLayers() >= 1 && FullMesh.Attributes()->NumNormalLayers() == 3))
@@ -720,6 +765,24 @@ int32 CutWithMesh(
 			}
 		}
 	}
+	return DynamicCuttingMesh;
+}
+
+int32 CutWithMesh(
+	const FDynamicMesh3& DynamicCuttingMesh,
+	FTransform CuttingMeshTransform,
+	FInternalSurfaceMaterials& InternalSurfaceMaterials,
+	FGeometryCollection& Collection,
+	const TArrayView<const int32>& TransformIndices,
+	double CollisionSampleSpacing,
+	const TOptional<FTransform>& TransformCollection,
+	bool bSetDefaultInternalMaterialsFromCollection,
+	FProgressCancel* Progress
+)
+{
+	FProgressCancel::FProgressScope PrepareScope = FProgressCancel::CreateScopeTo(Progress, .1);
+
+	int32 NewGeomStartIdx = -1;
 
 	if (bSetDefaultInternalMaterialsFromCollection)
 	{
@@ -730,15 +793,37 @@ int32 CutWithMesh(
 
 	FTransform CollectionToWorld = TransformCollection.Get(FTransform::Identity);
 
+	if (Progress && Progress->Cancelled())
+	{
+		return -1;
+	}
+
 	FDynamicMeshCollection MeshCollection(&Collection, TransformIndices, CollectionToWorld);
+	int32 NumUVLayers = Collection.NumUVLayers();
 	FCellMeshes CellMeshes(NumUVLayers, DynamicCuttingMesh, InternalSurfaceMaterials, CuttingMeshTransform);
 
 	TArray<TPair<int32, int32>> CellConnectivity;
 	CellConnectivity.Add(TPair<int32, int32>(0, -1)); // there's only one 'inside' cell (0), so all cut surfaces are connecting the 'inside' cell (0) to the 'outside' cell (-1)
 
+	PrepareScope.Done();
+	if (Progress && Progress->Cancelled())
+	{
+		return -1;
+	}
+	FProgressCancel::FProgressScope CutScope = FProgressCancel::CreateScopeTo(Progress, .99);
+
 	NewGeomStartIdx = MeshCollection.CutWithCellMeshes(InternalSurfaceMaterials, CellConnectivity, CellMeshes, &Collection, bSetDefaultInternalMaterialsFromCollection, CollisionSampleSpacing);
 
+	CutScope.Done();
+	if (Progress && Progress->Cancelled())
+	{
+		return -1;
+	}
+	FProgressCancel::FProgressScope ReindexScope = FProgressCancel::CreateScopeTo(Progress, 1);
+
 	Collection.ReindexMaterials();
+	ReindexScope.Done();
+
 	return NewGeomStartIdx;
 }
 
diff --git a/Engine/Plugins/Experimental/PlanarCutPlugin/Source/PlanarCut/Public/PlanarCut.h b/Engine/Plugins/Experimental/PlanarCutPlugin/Source/PlanarCut/Public/PlanarCut.h
--- a/Engine/Plugins/Experimental/PlanarCutPlugin/Source/PlanarCut/Public/PlanarCut.h
+++ b/Engine/Plugins/Experimental/PlanarCutPlugin/Source/PlanarCut/Public/PlanarCut.h
@@ -7,6 +7,8 @@
 #include "Voronoi/Voronoi.h"
 #include "GeometryCollection/GeometryCollection.h"
 #include "MeshDescription.h"
+#include "Util/ProgressCancel.h"
+#include "DynamicMesh/DynamicMesh3.h"
 
 struct PLANARCUT_API FNoiseSettings
 {
@@ -148,6 +150,7 @@ struct PLANARCUT_API FPlanarCells
  * @param RandomSeed				Seed to be used for random noise displacement
  * @param TransformCollection		Optional transform of the whole geometry collection; if unset, defaults to Identity
  * @param bIncludeOutsideCellInOutput	If true, geometry that was not inside any of the cells (e.g. was outside of the bounds of all cutting geometry) will still be included in the output; if false, it will be discarded.
+ * @param Progress						Optionally tracks progress and supports early-cancel
  * @return	index of first new geometry in the Output GeometryCollection, or -1 if no geometry was added
  */
 int32 PLANARCUT_API CutWithPlanarCells(
@@ -159,7 +162,8 @@ int32 PLANARCUT_API CutWithPlanarCells(
 	int32 RandomSeed,
 	const TOptional<FTransform>& TransformCollection = TOptional<FTransform>(),
 	bool bIncludeOutsideCellInOutput = true,
-	bool bSetDefaultInternalMaterialsFromCollection = true
+	bool bSetDefaultInternalMaterialsFromCollection = true,
+	FProgressCancel* Progress = nullptr
 );
 
 /**
@@ -173,6 +177,7 @@ int32 PLANARCUT_API CutWithPlanarCells(
  * @param RandomSeed				Seed to be used for random noise displacement
  * @param TransformCollection		Optional transform of the whole geometry collection; if unset, defaults to Identity
  * @param bIncludeOutsideCellInOutput	If true, geometry that was not inside any of the cells (e.g. was outside of the bounds of all cutting geometry) will still be included in the output; if false, it will be discarded.
+ * @param Progress						Optionally tracks progress and supports early-cancel
  * @return	index of first new geometry in the Output GeometryCollection, or -1 if no geometry was added
  */
 int32 PLANARCUT_API CutMultipleWithPlanarCells(
@@ -184,7 +189,8 @@ int32 PLANARCUT_API CutMultipleWithPlanarCells(
 	int32 RandomSeed,
 	const TOptional<FTransform>& TransformCollection = TOptional<FTransform>(),
 	bool bIncludeOutsideCellInOutput = true,
-	bool bSetDefaultInternalMaterialsFromCollection = true
+	bool bSetDefaultInternalMaterialsFromCollection = true,
+	FProgressCancel* Progress = nullptr
 );
 
 /**
@@ -198,6 +204,7 @@ int32 PLANARCUT_API CutMultipleWithPlanarCells(
  * @param CollisionSampleSpacing	Target spacing between collision sample vertices
  * @param RandomSeed				Seed to be used for random noise displacement
  * @param TransformCollection		Optional transform of the whole geometry collection; if unset, defaults to Identity
+ * @param Progress					Optionally tracks progress and supports early-cancel
  * @return	index of first new geometry in the Output GeometryCollection, or -1 if no geometry was added
  */
 int32 PLANARCUT_API CutMultipleWithMultiplePlanes(
@@ -209,7 +216,8 @@ int32 PLANARCUT_API CutMultipleWithMultiplePlanes(
 	double CollisionSampleSpacing,
 	int32 RandomSeed,
 	const TOptional<FTransform>& TransformCollection = TOptional<FTransform>(),
-	bool bSetDefaultInternalMaterialsFromCollection = true
+	bool bSetDefaultInternalMaterialsFromCollection = true,
+	FProgressCancel* Progress = nullptr
 );
 
 
@@ -329,19 +337,51 @@ int32 PLANARCUT_API AddCollisionSampleVertices(double TargetSpacing, FGeometryCo
  * @param TransformIndices			Which transform groups inside the collection to cut
  * @param CollisionSampleSpacing	Target spacing between collision sample vertices
  * @param TransformCollection		Optional transform of the collection; if unset, defaults to Identity
+ * @param Progress					Optionally tracks progress and supports early-cancel
  * @return index of first new geometry in the Output GeometryCollection, or -1 if no geometry was added
  */
 int32 PLANARCUT_API CutWithMesh(
-	FMeshDescription* CuttingMesh,
+	const UE::Geometry::FDynamicMesh3& CuttingMesh,
 	FTransform CuttingMeshTransform,
 	FInternalSurfaceMaterials& InternalSurfaceMaterials,
 	FGeometryCollection& Collection,
 	const TArrayView<const int32>& TransformIndices,
 	double CollisionSampleSpacing,
 	const TOptional<FTransform>& TransformCollection = TOptional<FTransform>(),
-	bool bSetDefaultInternalMaterialsFromCollection = true
+	bool bSetDefaultInternalMaterialsFromCollection = true,
+	FProgressCancel* Progress = nullptr
 );
 
+/// Convert a mesh description to a dynamic mesh *specifically* augmented / designed for cutting geometry collections, to be passed to CutWithMesh
+UE::Geometry::FDynamicMesh3 PLANARCUT_API ConvertMeshDescriptionToCuttingDynamicMesh(const FMeshDescription* CuttingMesh, int32 NumUVLayers, FProgressCancel* Progress = nullptr);
+
+inline int32 CutWithMesh(
+	const FMeshDescription* CuttingMesh,
+	FTransform CuttingMeshTransform,
+	FInternalSurfaceMaterials& InternalSurfaceMaterials,
+	FGeometryCollection& Collection,
+	const TArrayView<const int32>& TransformIndices,
+	double CollisionSampleSpacing,
+	const TOptional<FTransform>& TransformCollection = TOptional<FTransform>(),
+	bool bSetDefaultInternalMaterialsFromCollection = true,
+	FProgressCancel* Progress = nullptr
+)
+{
+	int32 NumUVLayers = Collection.NumUVLayers();
+	return CutWithMesh(ConvertMeshDescriptionToCuttingDynamicMesh(CuttingMesh, NumUVLayers, Progress),
+		CuttingMeshTransform,
+		InternalSurfaceMaterials,
+		Collection,
+		TransformIndices,
+		CollisionSampleSpacing,
+		TransformCollection,
+		bSetDefaultInternalMaterialsFromCollection,
+		Progress
+		);
+}
+
+
+
 /**
  * Convert chosen Geometry groups inside a GeometryCollection to a single Mesh Description.
  *
