add missing #undef LOCTEXT_NAMESPACE for planarcut.cpp

diff --git a/Engine/Plugins/Experimental/PlanarCutPlugin/Source/PlanarCut/Private/PlanarCut.cpp b/Engine/Plugins/Experimental/PlanarCutPlugin/Source/PlanarCut/Private/PlanarCut.cpp
--- a/Engine/Plugins/Experimental/PlanarCutPlugin/Source/PlanarCut/Private/PlanarCut.cpp
+++ b/Engine/Plugins/Experimental/PlanarCutPlugin/Source/PlanarCut/Private/PlanarCut.cpp
@@ -1402,4 +1402,6 @@ void ConvertToMeshDescription(
 
 	FDynamicMeshToMeshDescription Converter;
 	Converter.Convert(&CombinedMesh, MeshOut, true);
-}
\ No newline at end of file
+}
+
+#undef LOCTEXT_NAMESPACE
