expand progresscancel usage to fracture plane cuts

diff --git a/Engine/Plugins/Experimental/ChaosEditor/Source/FractureEditor/Private/FractureToolPlaneCut.cpp b/Engine/Plugins/Experimental/ChaosEditor/Source/FractureEditor/Private/FractureToolPlaneCut.cpp
--- a/Engine/Plugins/Experimental/ChaosEditor/Source/FractureEditor/Private/FractureToolPlaneCut.cpp
+++ b/Engine/Plugins/Experimental/ChaosEditor/Source/FractureEditor/Private/FractureToolPlaneCut.cpp
@@ -9,6 +9,9 @@
 #include "FractureToolContext.h"
 #include "Drawing/MeshDebugDrawing.h"
 #include "FrameTypes.h"
+#include "FractureToolBackgroundTask.h"
+
+using namespace UE::Fracture;
 
 #define LOCTEXT_NAMESPACE "FracturePlanar"
 
@@ -164,17 +167,53 @@ void UFractureToolPlaneCut::FractureContextChanged()
 	}
 }
 
+
+class FPlaneFractureOp : public FGeometryCollectionFractureOperator
+{
+public:
+	FPlaneFractureOp(const FGeometryCollection& SourceCollection) : FGeometryCollectionFractureOperator(SourceCollection)
+	{}
+
+	virtual ~FPlaneFractureOp() = default;
+
+	TArray<int> Selection;
+	TArray<FPlane> CuttingPlanes;
+	FInternalSurfaceMaterials InternalSurfaceMaterials;
+	float PointSpacing;
+	float Grout;
+	int Seed;
+	FTransform Transform;
+	UE::Geometry::FDynamicMesh3 CuttingMesh;
+
+	// TGenericDataOperator interface:
+	virtual void CalculateResult(FProgressCancel* Progress) override
+	{
+		ResultGeometryIndex = CutMultipleWithMultiplePlanes(CuttingPlanes, InternalSurfaceMaterials, *CollectionCopy, Selection, Grout, PointSpacing, Seed, Transform, true, Progress);
+
+		if (Progress && Progress->Cancelled())
+		{
+			return;
+		}
+		SetResult(MoveTemp(CollectionCopy));
+	}
+};
+
+
 int32 UFractureToolPlaneCut::ExecuteFracture(const FFractureToolContext& FractureContext)
 {
 	if (FractureContext.IsValid())
 	{
-		const UFracturePlaneCutSettings* LocalCutSettings = PlaneCutSettings;
+		TUniquePtr<FPlaneFractureOp> PlaneCutOp = MakeUnique<FPlaneFractureOp>(*(FractureContext.GetGeometryCollection()));
+		PlaneCutOp->Selection = FractureContext.GetSelection();
+		PlaneCutOp->Grout = CutterSettings->Grout;
+		PlaneCutOp->PointSpacing = CollisionSettings->GetPointSpacing();
+		PlaneCutOp->Seed = FractureContext.GetSeed();
+		PlaneCutOp->Transform = FractureContext.GetTransform();
 
-		TArray<FPlane> CuttingPlanes;
 		if (GizmoSettings->IsGizmoEnabled())
 		{
 			FTransform Transform = GizmoSettings->GetTransform();
-			CuttingPlanes.Add(FPlane(Transform.GetLocation(), Transform.GetUnitAxis(EAxis::Z)));
+			PlaneCutOp->CuttingPlanes.Add(FPlane(Transform.GetLocation(), Transform.GetUnitAxis(EAxis::Z)));
 		}
 		else
 		{
@@ -182,22 +221,20 @@ int32 UFractureToolPlaneCut::ExecuteFracture(const FFractureToolContext& Fractur
 			GenerateSliceTransforms(FractureContext, CuttingPlaneTransforms);
 			for (const FTransform& Transform : CuttingPlaneTransforms)
 			{
-				CuttingPlanes.Add(FPlane(Transform.GetLocation(), Transform.GetUnitAxis(EAxis::Z)));
+				PlaneCutOp->CuttingPlanes.Add(FPlane(Transform.GetLocation(), Transform.GetUnitAxis(EAxis::Z)));
 			}
 		}
-
-		FInternalSurfaceMaterials InternalSurfaceMaterials;
+		
 		FNoiseSettings NoiseSettings;
 		if (CutterSettings->Amplitude > 0.0f)
 		{
 			CutterSettings->TransferNoiseSettings(NoiseSettings);
-			InternalSurfaceMaterials.NoiseSettings = NoiseSettings;
+			PlaneCutOp->InternalSurfaceMaterials.NoiseSettings = NoiseSettings;
 		}
 
-		// Proximity is invalidated.
-		ClearProximity(FractureContext.GetGeometryCollection().Get());
-
-		return CutMultipleWithMultiplePlanes(CuttingPlanes, InternalSurfaceMaterials, *FractureContext.GetGeometryCollection(), FractureContext.GetSelection(), CutterSettings->Grout, CollisionSettings->GetPointSpacing(), FractureContext.GetSeed(), FractureContext.GetTransform());
+		int Result = RunCancellableGeometryCollectionOp<FPlaneFractureOp>(*(FractureContext.GetGeometryCollection()),
+			MoveTemp(PlaneCutOp), LOCTEXT("ComputingPlaneFractureMessage", "Computing Plane Fracture"));
+		return Result;
 	}
 
 	return INDEX_NONE;
diff --git a/Engine/Plugins/Experimental/PlanarCutPlugin/Source/PlanarCut/Private/GeometryMeshConversion.cpp b/Engine/Plugins/Experimental/PlanarCutPlugin/Source/PlanarCut/Private/GeometryMeshConversion.cpp
--- a/Engine/Plugins/Experimental/PlanarCutPlugin/Source/PlanarCut/Private/GeometryMeshConversion.cpp
+++ b/Engine/Plugins/Experimental/PlanarCutPlugin/Source/PlanarCut/Private/GeometryMeshConversion.cpp
@@ -36,7 +36,7 @@
 #include "DynamicMesh/MeshNormals.h"
 #include "DynamicMesh/MeshTangents.h"
 #include "ConstrainedDelaunay2.h"
-
+#include "Util/ProgressCancel.h"
 
 #include "StaticMeshOperations.h"
 #include "MeshDescriptionToDynamicMesh.h"
@@ -44,10 +44,6 @@
 
 #include "Algo/Rotate.h"
 
-#if WITH_EDITOR
-#include "Misc/ScopedSlowTask.h"
-#endif
-
 using namespace UE::Geometry;
 
 
@@ -2245,31 +2241,18 @@ int32 FDynamicMeshCollection::CutWithMultiplePlanes(
 	int32 RandomSeed,
 	FGeometryCollection* Collection,
 	FInternalSurfaceMaterials& InternalSurfaceMaterials,
-	bool bSetDefaultInternalMaterialsFromCollection
+	bool bSetDefaultInternalMaterialsFromCollection,
+	FProgressCancel* Progress
 )
 {
-#if WITH_EDITOR
-	// Create progress indicator dialog
-	static const FText SlowTaskText = NSLOCTEXT("CutMultipleWithMultiplePlanes", "CutMultipleWithMultiplePlanesText", "Cutting geometry collection with plane(s)...");
-
-	FScopedSlowTask SlowTask(Planes.Num(), SlowTaskText);
-	SlowTask.MakeDialog();
-
-	// Declare progress shortcut lambdas
-	auto EnterProgressFrame = [&SlowTask](float Progress)
-	{
-		SlowTask.EnterProgressFrame(Progress);
-	};
-#else
-	auto EnterProgressFrame = [](float Progress) {};
-#endif
-
 	bool bHasGrout = Grout > 0;
 
 	int32 NumUVLayers = Collection->NumUVLayers();
 
 	FRandomStream RandomStream(RandomSeed);
 
+	const float PerPlaneWorkFrac = 1.0f / Planes.Num();
+
 	if (bHasGrout)
 	{
 		// For multi-plane cuts with grout specifically, the easiest path seems to be:
@@ -2285,26 +2268,34 @@ int32 FDynamicMeshCollection::CutWithMultiplePlanes(
 		FMeshIndexMappings IndexMaps;
 		for (int32 PlaneIdx = 0; PlaneIdx < Planes.Num(); PlaneIdx++)
 		{
-			EnterProgressFrame(.5);
+			FProgressCancel::FProgressScope MakeMeshScope(Progress, .1 * PerPlaneWorkFrac);
 			FPlanarCells PlaneCells(Planes[PlaneIdx]);
 			PlaneCells.InternalSurfaceMaterials = InternalSurfaceMaterials;
 			FCellMeshes PlaneGroutMesh(NumUVLayers, RandomStream);
 			PlaneGroutMesh.MakeOnlyPlanarGroutCell(PlaneCells, Bounds, Grout);
 			GroutAppender.AppendMesh(&PlaneGroutMesh.CellMeshes[0].AugMesh, IndexMaps);
+			if (Progress && Progress->Cancelled())
+			{
+				return INDEX_NONE;
+			}
 		}
 
-		EnterProgressFrame(Planes.Num() * .2);
+		FProgressCancel::FProgressScope GroutUnionScope(Progress, .3);
 		FMeshSelfUnion GroutUnion(&GroutMesh);
 		GroutUnion.bSimplifyAlongNewEdges = true;
 		GroutUnion.bWeldSharedEdges = false;
 		GroutUnion.Compute();
-
-		EnterProgressFrame(Planes.Num() * .1);
+		GroutUnionScope.Done();
+		if (Progress && Progress->Cancelled())
+		{
+			return INDEX_NONE;
+		}
+		
+		FProgressCancel::FProgressScope GroutCutScope = FProgressCancel::CreateScopeTo(Progress, 1);
 		// first mesh is the same as the second mesh, but will be subtracted b/c it's the "outside cell"
 		GroutCells.CellMeshes[1].AugMesh = GroutMesh;
 		GroutCells.OutsideCellIndex = 1;
 
-		EnterProgressFrame(Planes.Num() * .2);
 		TArray<TPair<int32, int32>> CellConnectivity;
 		CellConnectivity.Add(TPair<int32, int32>(0, -1));
 
@@ -2319,7 +2310,11 @@ int32 FDynamicMeshCollection::CutWithMultiplePlanes(
 	}
 	for (int32 PlaneIdx = 0; PlaneIdx < Planes.Num(); PlaneIdx++)
 	{
-		EnterProgressFrame(1);
+		if (Progress && Progress->Cancelled())
+		{
+			return INDEX_NONE;
+		}
+		FProgressCancel::FProgressScope OnePlaneScope(Progress, PerPlaneWorkFrac);
 		FPlanarCells PlaneCells(Planes[PlaneIdx]);
 		PlaneCells.InternalSurfaceMaterials = InternalSurfaceMaterials;
 		double OnePercentExtend = Bounds.MaxDim() * .01;
@@ -2452,6 +2447,11 @@ int32 FDynamicMeshCollection::CutWithMultiplePlanes(
 		}
 	}
 
+	if (Progress && Progress->Cancelled())
+	{
+		return INDEX_NONE;
+	}
+
 	int32 NewFirstIndex = FirstCreatedIndex;
 
 	constexpr bool bRemoveOldGeometry = false; // if false, we just hide the geometry that we've replaced by fractured child geometry, rather than remove it
diff --git a/Engine/Plugins/Experimental/PlanarCutPlugin/Source/PlanarCut/Private/GeometryMeshConversion.h b/Engine/Plugins/Experimental/PlanarCutPlugin/Source/PlanarCut/Private/GeometryMeshConversion.h
--- a/Engine/Plugins/Experimental/PlanarCutPlugin/Source/PlanarCut/Private/GeometryMeshConversion.h
+++ b/Engine/Plugins/Experimental/PlanarCutPlugin/Source/PlanarCut/Private/GeometryMeshConversion.h
@@ -16,6 +16,8 @@ struct FPlanarCells;
 struct FInternalSurfaceMaterials;
 struct FNoiseSettings;
 
+class FProgressCancel;
+
 namespace UE
 {
 namespace PlanarCut
@@ -223,7 +225,8 @@ struct PLANARCUT_API FDynamicMeshCollection
 		int32 RandomSeed,
 		FGeometryCollection* Collection,
 		FInternalSurfaceMaterials& InternalSurfaceMaterials,
-		bool bSetDefaultInternalMaterialsFromCollection
+		bool bSetDefaultInternalMaterialsFromCollection,
+		FProgressCancel* Progress = nullptr
 	);
 
 	/**
diff --git a/Engine/Plugins/Experimental/PlanarCutPlugin/Source/PlanarCut/Private/PlanarCut.cpp b/Engine/Plugins/Experimental/PlanarCutPlugin/Source/PlanarCut/Private/PlanarCut.cpp
--- a/Engine/Plugins/Experimental/PlanarCutPlugin/Source/PlanarCut/Private/PlanarCut.cpp
+++ b/Engine/Plugins/Experimental/PlanarCutPlugin/Source/PlanarCut/Private/PlanarCut.cpp
@@ -33,7 +33,7 @@
 #include "DynamicMesh/MeshNormals.h"
 #include "DynamicMesh/MeshTangents.h"
 #include "ConstrainedDelaunay2.h"
-
+#include "Util/ProgressCancel.h"
 
 #include "StaticMeshOperations.h"
 #include "MeshDescriptionToDynamicMesh.h"
@@ -43,10 +43,6 @@
 
 #include "GeometryMeshConversion.h"
 
-#if WITH_EDITOR
-#include "Misc/ScopedSlowTask.h"
-#endif
-
 using namespace UE::Geometry;
 
 
@@ -587,10 +583,10 @@ int32 CutMultipleWithMultiplePlanes(
 	{
 		return -1;
 	}
-	FProgressCancel::FProgressScope CutScope = FProgressCancel::CreateScopeTo(Progress, .99);
+	FProgressCancel::FProgressScope CutScope = FProgressCancel::CreateScopeTo(Progress, .99, LOCTEXT("CutWithMultiplePlanesBody", "Cutting with planes"));
 
 	int32 NewGeomStartIdx = -1;
-	NewGeomStartIdx = MeshCollection.CutWithMultiplePlanes(Planes, Grout, CollisionSampleSpacing, RandomSeed, &Collection, InternalSurfaceMaterials, bSetDefaultInternalMaterialsFromCollection);
+	NewGeomStartIdx = MeshCollection.CutWithMultiplePlanes(Planes, Grout, CollisionSampleSpacing, RandomSeed, &Collection, InternalSurfaceMaterials, bSetDefaultInternalMaterialsFromCollection, Progress);
 
 	CutScope.Done();
 	if (Progress && Progress->Cancelled())
diff --git a/Engine/Plugins/Experimental/PlanarCutPlugin/Source/PlanarCut/Public/PlanarCut.h b/Engine/Plugins/Experimental/PlanarCutPlugin/Source/PlanarCut/Public/PlanarCut.h
--- a/Engine/Plugins/Experimental/PlanarCutPlugin/Source/PlanarCut/Public/PlanarCut.h
+++ b/Engine/Plugins/Experimental/PlanarCutPlugin/Source/PlanarCut/Public/PlanarCut.h
@@ -7,9 +7,10 @@
 #include "Voronoi/Voronoi.h"
 #include "GeometryCollection/GeometryCollection.h"
 #include "MeshDescription.h"
-#include "Util/ProgressCancel.h"
 #include "DynamicMesh/DynamicMesh3.h"
 
+class FProgressCancel;
+
 struct PLANARCUT_API FNoiseSettings
 {
 	float Amplitude = 2;
