Chaos: add a mesh facade for geometry collection

diff --git a/Engine/Plugins/Experimental/PlanarCutPlugin/Source/PlanarCut/Private/PlanarCut.cpp b/Engine/Plugins/Experimental/PlanarCutPlugin/Source/PlanarCut/Private/PlanarCut.cpp
--- a/Engine/Plugins/Experimental/PlanarCutPlugin/Source/PlanarCut/Private/PlanarCut.cpp
+++ b/Engine/Plugins/Experimental/PlanarCutPlugin/Source/PlanarCut/Private/PlanarCut.cpp
@@ -93,26 +93,32 @@ int32 FInternalSurfaceMaterials::GetDefaultMaterialIDForGeometry(const FGeometry
 	return InternalMaterialID;
 }
 
-void FInternalSurfaceMaterials::SetUVScaleFromCollection(const FGeometryCollection& Collection, int32 GeometryIdx)
+void FInternalSurfaceMaterials::SetUVScaleFromCollection(const FGeometryCollectionMeshFacade& CollectionMesh, int32 GeometryIdx)
 {
+	const auto& VertexArray = CollectionMesh.Vertex.Get();
+	const auto& UVsArray = CollectionMesh.UVs.Get();
+	const auto& IndicesArray = CollectionMesh.Indices.Get();
+	const auto& FaceStartArray = CollectionMesh.FaceStart.Get();
+	const auto& FaceCountArray = CollectionMesh.FaceCount.Get();
+	
 	int32 FaceStart = 0;
-	int32 FaceEnd = Collection.Indices.Num();
+	int32 FaceEnd = IndicesArray.Num();
 	if (GeometryIdx > -1)
 	{
-		FaceStart = Collection.FaceStart[GeometryIdx];
-		FaceEnd = Collection.FaceCount[GeometryIdx] + Collection.FaceStart[GeometryIdx];
+		FaceStart = FaceStartArray[GeometryIdx];
+		FaceEnd = FaceCountArray[GeometryIdx] + FaceStartArray[GeometryIdx];
 	}
 	float UVDistance = 0;
 	float WorldDistance = 0;
 	for (int32 FaceIdx = FaceStart; FaceIdx < FaceEnd; FaceIdx++)
 	{
-		const FIntVector& Tri = Collection.Indices[FaceIdx];
-		WorldDistance += FVector3f::Distance(Collection.Vertex[Tri.X], Collection.Vertex[Tri.Y]);
-		UVDistance += FVector2D::Distance(FVector2D(Collection.UVs[Tri.X][0]), FVector2D(Collection.UVs[Tri.Y][0]));
-		WorldDistance += FVector3f::Distance(Collection.Vertex[Tri.Z], Collection.Vertex[Tri.Y]);
-		UVDistance += FVector2D::Distance(FVector2D(Collection.UVs[Tri.Z][0]), FVector2D(Collection.UVs[Tri.Y][0]));
-		WorldDistance += FVector3f::Distance(Collection.Vertex[Tri.X], Collection.Vertex[Tri.Z]);
-		UVDistance += FVector2D::Distance(FVector2D(Collection.UVs[Tri.X][0]), FVector2D(Collection.UVs[Tri.Z][0]));
+		const FIntVector& Tri = IndicesArray[FaceIdx];
+		WorldDistance += FVector3f::Distance(VertexArray[Tri.X], VertexArray[Tri.Y]);
+		UVDistance += FVector2D::Distance(FVector2D(UVsArray[Tri.X][0]), FVector2D(UVsArray[Tri.Y][0]));
+		WorldDistance += FVector3f::Distance(VertexArray[Tri.Z], VertexArray[Tri.Y]);
+		UVDistance += FVector2D::Distance(FVector2D(UVsArray[Tri.Z][0]), FVector2D(UVsArray[Tri.Y][0]));
+		WorldDistance += FVector3f::Distance(VertexArray[Tri.X], VertexArray[Tri.Z]);
+		UVDistance += FVector2D::Distance(FVector2D(UVsArray[Tri.X][0]), FVector2D(UVsArray[Tri.Z][0]));
 	}
 
 	if (WorldDistance > 0)
@@ -569,9 +575,15 @@ int32 CutMultipleWithMultiplePlanes(
 	int32 OrigNumGeom = Collection.FaceCount.Num();
 	int32 CurNumGeom = OrigNumGeom;
 
+	FGeometryCollectionMeshFacade CollectionMesh(Collection);
+	if (!CollectionMesh.IsValid())
+	{
+		return -1;
+	}
+
 	if (bSetDefaultInternalMaterialsFromCollection)
 	{
-		InternalSurfaceMaterials.SetUVScaleFromCollection(Collection);
+		InternalSurfaceMaterials.SetUVScaleFromCollection(CollectionMesh);
 	}
 
 	FTransform CollectionToWorld = TransformCollection.Get(FTransform::Identity);
diff --git a/Engine/Plugins/Experimental/PlanarCutPlugin/Source/PlanarCut/Public/PlanarCut.h b/Engine/Plugins/Experimental/PlanarCutPlugin/Source/PlanarCut/Public/PlanarCut.h
--- a/Engine/Plugins/Experimental/PlanarCutPlugin/Source/PlanarCut/Public/PlanarCut.h
+++ b/Engine/Plugins/Experimental/PlanarCutPlugin/Source/PlanarCut/Public/PlanarCut.h
@@ -36,7 +36,7 @@ struct PLANARCUT_API FInternalSurfaceMaterials
 	 * @param Collection	Reference collection to use for setting UV scale
 	 * @param GeometryIdx	Reference geometry inside collection; if -1, use all geometries in collection
 	 */
-	void SetUVScaleFromCollection(const FGeometryCollection& Collection, int32 GeometryIdx = -1);
+	void SetUVScaleFromCollection(const FGeometryCollectionMeshFacade& CollectionMesh, int32 GeometryIdx = -1);
 
 	
 	int32 GetDefaultMaterialIDForGeometry(const FGeometryCollection& Collection, int32 GeometryIdx = -1) const;
diff --git a/Engine/Source/Runtime/Experimental/Chaos/Private/GeometryCollection/GeometryCollection.cpp b/Engine/Source/Runtime/Experimental/Chaos/Private/GeometryCollection/GeometryCollection.cpp
--- a/Engine/Source/Runtime/Experimental/Chaos/Private/GeometryCollection/GeometryCollection.cpp
+++ b/Engine/Source/Runtime/Experimental/Chaos/Private/GeometryCollection/GeometryCollection.cpp
@@ -1959,3 +1959,62 @@ void FGeometryCollection::UpdateOldAttributeNames()
 	}
 	this->RemoveGroup("Structure");
 }
+
+FGeometryCollectionMeshFacade::FGeometryCollectionMeshFacade(FManagedArrayCollection& InCollection)
+	: Vertex(InCollection, "Vertex", FGeometryCollection::VerticesGroup)
+	, TangentU(InCollection, "TangentU", FGeometryCollection::VerticesGroup)
+	, TangentV(InCollection, "TangentV", FGeometryCollection::VerticesGroup)
+	, Normal(InCollection, "Normal", FGeometryCollection::VerticesGroup)
+	, UVs(InCollection, "UVs", FGeometryCollection::VerticesGroup)
+	, Color(InCollection, "Color", FGeometryCollection::VerticesGroup)
+	, BoneMap(InCollection, "BoneMap", FGeometryCollection::VerticesGroup)
+	, VertexStart(InCollection, "VertexStart", FGeometryCollection::GeometryGroup)
+	, VertexCount(InCollection, "VertexCount", FGeometryCollection::GeometryGroup)
+	, Indices(InCollection, "Indices", FGeometryCollection::FacesGroup)
+	, Visible(InCollection, "Visible", FGeometryCollection::FacesGroup)
+	, MaterialIndex(InCollection, "MaterialIndex", FGeometryCollection::FacesGroup)
+	, MaterialID(InCollection, "MaterialID", FGeometryCollection::FacesGroup)
+	, FaceStart(InCollection, "FaceStart", FGeometryCollection::GeometryGroup)
+	, FaceCount(InCollection, "FaceCount", FGeometryCollection::GeometryGroup)
+{
+}
+
+bool FGeometryCollectionMeshFacade::IsValid() const
+{
+	return Vertex.IsValid()
+		&& TangentU.IsValid()
+		&& TangentV.IsValid()
+		&& Normal.IsValid()
+		&& UVs.IsValid()
+		&& Color.IsValid()
+		&& BoneMap.IsValid()
+		&& VertexStart.IsValid()
+		&& VertexCount.IsValid()
+		&& Indices.IsValid()
+		&& Visible.IsValid()
+		&& MaterialIndex.IsValid()
+		&& MaterialID.IsValid()
+		&& FaceStart.IsValid()
+		&& FaceCount.IsValid()
+	;
+}
+
+void FGeometryCollectionMeshFacade::AddAttributes()
+{
+	Vertex.Add();
+	TangentU.Add();
+	TangentV.Add();
+	Normal.Add();
+	UVs.Add();
+	Color.Add();
+	BoneMap.Add();
+	VertexStart.Add();
+	VertexCount.Add();
+	Indices.Add();
+	Visible.Add();
+	MaterialIndex.Add();
+	MaterialID.Add();
+	FaceStart.Add();
+	FaceCount.Add();
+}
+
diff --git a/Engine/Source/Runtime/Experimental/Chaos/Public/GeometryCollection/GeometryCollection.h b/Engine/Source/Runtime/Experimental/Chaos/Public/GeometryCollection/GeometryCollection.h
--- a/Engine/Source/Runtime/Experimental/Chaos/Public/GeometryCollection/GeometryCollection.h
+++ b/Engine/Source/Runtime/Experimental/Chaos/Public/GeometryCollection/GeometryCollection.h
@@ -343,3 +343,75 @@ FORCEINLINE Chaos::FChaosArchive& operator<<(Chaos::FChaosArchive& Ar, FGeometry
 	Value.Serialize(Ar);
 	return Ar;
 }
+
+
+template <typename T>
+struct AttributeAccessor
+{
+public:
+	AttributeAccessor(FManagedArrayCollection& InCollection, const FName& AttributeName, const FName& AttributeGroup)
+		: Collection(InCollection)
+		, Name(AttributeName)
+		, Group(AttributeGroup)
+	{
+		AttributeArray = InCollection.FindAttributeTyped<T>(AttributeName, AttributeGroup);
+
+	}
+	bool IsValid() const { return AttributeArray != nullptr; }
+	const TManagedArray<T>& Get() const
+	{
+		check(IsValid());
+		return *AttributeArray;
+	}
+	TManagedArray<T>& Modify() const
+	{
+		check(IsValid());
+		AttributeArray->MarkDirty();
+		return *AttributeArray;
+	}
+	void Add()
+	{
+		AttributeArray = &Collection.AddAttribute<T>(Name, Group);
+	}
+
+private:
+	FManagedArrayCollection& Collection;
+	FName Name;
+	FName Group;
+	TManagedArray<T>* AttributeArray;
+};
+
+
+class CHAOS_API FGeometryCollectionMeshFacade
+{
+public:
+	FGeometryCollectionMeshFacade(FManagedArrayCollection& InCollection);
+
+	/** 
+	 * returns true iof all the necessary attributes are present
+	 * if not then the API can be used to create  
+	 */
+	bool IsValid() const;	
+
+	/**
+	 * Add the necessary attributes if they are missing
+	 */
+	void AddAttributes();
+
+	AttributeAccessor<FVector3f> Vertex;
+	AttributeAccessor<FVector3f> TangentU;
+	AttributeAccessor<FVector3f> TangentV;
+	AttributeAccessor<FVector3f> Normal;
+	AttributeAccessor<TArray<FVector2f>> UVs;
+	AttributeAccessor<FLinearColor> Color;
+	AttributeAccessor<int32> BoneMap;
+	AttributeAccessor<int32> VertexStart;
+	AttributeAccessor<int32> VertexCount;
+
+	AttributeAccessor<FIntVector> Indices;
+	AttributeAccessor<bool> Visible;
+	AttributeAccessor<int32> MaterialIndex;
+	AttributeAccessor<int32> MaterialID;
+	AttributeAccessor<int32> FaceStart;
+	AttributeAccessor<int32> FaceCount;
+};
\ No newline at end of file
