add fracture merge/split tools, to let user combine selected geometry or split it (into connected components)

diff --git a/Engine/Plugins/Experimental/ChaosEditor/Source/FractureEditor/Private/FractureEditorCommands.h b/Engine/Plugins/Experimental/ChaosEditor/Source/FractureEditor/Private/FractureEditorCommands.h
--- a/Engine/Plugins/Experimental/ChaosEditor/Source/FractureEditor/Private/FractureEditorCommands.h
+++ b/Engine/Plugins/Experimental/ChaosEditor/Source/FractureEditor/Private/FractureEditorCommands.h
@@ -60,6 +60,8 @@ class FFractureEditorCommands : public TCommands<FFractureEditorCommands>
 		TSharedPtr< FUICommandInfo > DeleteBranch;
 		TSharedPtr< FUICommandInfo > Hide;
 		TSharedPtr< FUICommandInfo > Unhide;
+		TSharedPtr< FUICommandInfo > MergeSelected;
+		TSharedPtr< FUICommandInfo > SplitSelected;
 		
 		// Generate Commands
 		TSharedPtr< FUICommandInfo > GenerateAsset;
diff --git a/Engine/Plugins/Experimental/ChaosEditor/Source/FractureEditor/Private/FractureEditorModeToolkit.cpp b/Engine/Plugins/Experimental/ChaosEditor/Source/FractureEditor/Private/FractureEditorModeToolkit.cpp
--- a/Engine/Plugins/Experimental/ChaosEditor/Source/FractureEditor/Private/FractureEditorModeToolkit.cpp
+++ b/Engine/Plugins/Experimental/ChaosEditor/Source/FractureEditor/Private/FractureEditorModeToolkit.cpp
@@ -845,6 +845,8 @@ void FFractureEditorModeToolkit::BuildToolPalette(FName PaletteIndex, class FToo
 		ToolbarBuilder.AddToolBarButton(Commands.DeleteBranch);
 		ToolbarBuilder.AddToolBarButton(Commands.Hide);
 		ToolbarBuilder.AddToolBarButton(Commands.Unhide);
+		ToolbarBuilder.AddToolBarButton(Commands.MergeSelected);
+		ToolbarBuilder.AddToolBarButton(Commands.SplitSelected);
 	}
 	else if (PaletteIndex == TEXT("Cluster"))
 	{
diff --git a/Engine/Plugins/Experimental/ChaosEditor/Source/FractureEditor/Private/FractureToolEditing.cpp b/Engine/Plugins/Experimental/ChaosEditor/Source/FractureEditor/Private/FractureToolEditing.cpp
--- a/Engine/Plugins/Experimental/ChaosEditor/Source/FractureEditor/Private/FractureToolEditing.cpp
+++ b/Engine/Plugins/Experimental/ChaosEditor/Source/FractureEditor/Private/FractureToolEditing.cpp
@@ -10,6 +10,8 @@
 #include "FractureToolContext.h"
 #include "ScopedTransaction.h"
 
+#include "PlanarCut.h"
+
 #define LOCTEXT_NAMESPACE "FractureToolEditing"
 
 
@@ -102,6 +104,119 @@ void UFractureToolDeleteBranch::Execute(TWeakPtr<FFractureEditorModeToolkit> InT
 
 
 
+FText UFractureToolMergeSelected::GetDisplayText() const
+{
+	return FText(NSLOCTEXT("FractureToolEditingOps", "FractureToolMergeSelected", "Merge"));
+}
+
+FText UFractureToolMergeSelected::GetTooltipText() const
+{
+	return FText(NSLOCTEXT("FractureToolEditingOps", "FractureToolMergeSelectedTooltip", "Merge all selected nodes into one node."));
+}
+
+FSlateIcon UFractureToolMergeSelected::GetToolIcon() const
+{
+	return FSlateIcon("FractureEditorStyle", "FractureEditor.MergeSelected");
+}
+
+void UFractureToolMergeSelected::RegisterUICommand(FFractureEditorCommands* BindingContext)
+{
+	UI_COMMAND_EXT(BindingContext, UICommandInfo, "MergeSelected", "Merge", "Merge all selected nodes into one node.", EUserInterfaceActionType::Button, FInputChord());
+	BindingContext->MergeSelected = UICommandInfo;
+}
+
+void UFractureToolMergeSelected::Execute(TWeakPtr<FFractureEditorModeToolkit> InToolkit)
+{
+	if (InToolkit.IsValid())
+	{
+		FScopedTransaction Transaction(LOCTEXT("MergeSelected", "Merge Selected"));
+
+		FFractureEditorModeToolkit* Toolkit = InToolkit.Pin().Get();
+
+		TArray<FFractureToolContext> Contexts = GetFractureToolContexts();
+
+		for (FFractureToolContext& Context : Contexts)
+		{
+			FGeometryCollectionEdit Edit(Context.GetGeometryCollectionComponent(), GeometryCollection::EEditUpdate::RestPhysicsDynamic);
+			FGeometryCollection* GeometryCollection = Context.GetGeometryCollection().Get();
+
+			Context.Sanitize();
+
+			const TArray<int32>& NodesForMerge = Context.GetSelection();
+
+			constexpr bool bBooleanUnion = false;
+			MergeAllSelectedBones(*GeometryCollection, NodesForMerge, bBooleanUnion);
+
+			// Proximity is invalidated.
+			ClearProximity(Context.GetGeometryCollection().Get());
+
+			Refresh(Context, Toolkit, true);
+		}
+
+		SetOutlinerComponents(Contexts, Toolkit);
+	}
+}
+
+
+FText UFractureToolSplitSelected::GetDisplayText() const
+{
+	return FText(NSLOCTEXT("FractureToolEditingOps", "FractureToolSplitSelected", "Split"));
+}
+
+FText UFractureToolSplitSelected::GetTooltipText() const
+{
+	return FText(NSLOCTEXT("FractureToolEditingOps", "FractureToolSplitSelectedTooltip", "Split all selected nodes into their connected component parts."));
+}
+
+FSlateIcon UFractureToolSplitSelected::GetToolIcon() const
+{
+	return FSlateIcon("FractureEditorStyle", "FractureEditor.SplitSelected");
+}
+
+void UFractureToolSplitSelected::RegisterUICommand(FFractureEditorCommands* BindingContext)
+{
+	UI_COMMAND_EXT(BindingContext, UICommandInfo, "SplitSelected", "Split", "Split all selected nodes into their connected component parts.", EUserInterfaceActionType::Button, FInputChord());
+	BindingContext->SplitSelected = UICommandInfo;
+}
+
+void UFractureToolSplitSelected::Execute(TWeakPtr<FFractureEditorModeToolkit> InToolkit)
+{
+	if (InToolkit.IsValid())
+	{
+		FScopedTransaction Transaction(LOCTEXT("SplitSelected", "Split Selected"));
+
+		FFractureEditorModeToolkit* Toolkit = InToolkit.Pin().Get();
+
+		TArray<FFractureToolContext> Contexts = GetFractureToolContexts();
+
+		for (FFractureToolContext& Context : Contexts)
+		{
+			FGeometryCollectionEdit Edit(Context.GetGeometryCollectionComponent(), GeometryCollection::EEditUpdate::RestPhysicsDynamic);
+			FGeometryCollection* GeometryCollection = Context.GetGeometryCollection().Get();
+
+			Context.Sanitize();
+
+			TArray<int32> NodesForSplit;
+
+			for (int32 Select : Context.GetSelection())
+			{
+				FGeometryCollectionClusteringUtility::GetLeafBones(GeometryCollection, Select, true, NodesForSplit);
+			}
+
+			SplitIslands(*GeometryCollection, NodesForSplit, 0, nullptr);
+
+			// Proximity is invalidated.
+			ClearProximity(Context.GetGeometryCollection().Get());
+
+			Refresh(Context, Toolkit, true);
+		}
+
+		SetOutlinerComponents(Contexts, Toolkit);
+	}
+}
+
+
+
 FText UFractureToolHide::GetDisplayText() const
 {
 	return FText(NSLOCTEXT("FractureToolEditingOps", "FractureToolHide", "Hide"));
diff --git a/Engine/Plugins/Experimental/ChaosEditor/Source/FractureEditor/Private/FractureToolEditing.h b/Engine/Plugins/Experimental/ChaosEditor/Source/FractureEditor/Private/FractureToolEditing.h
--- a/Engine/Plugins/Experimental/ChaosEditor/Source/FractureEditor/Private/FractureToolEditing.h
+++ b/Engine/Plugins/Experimental/ChaosEditor/Source/FractureEditor/Private/FractureToolEditing.h
@@ -22,6 +22,39 @@ public:
 	virtual void Execute(TWeakPtr<FFractureEditorModeToolkit> InToolkit) override;
 };
 
+UCLASS(DisplayName = "Merge Selected", Category = "FractureTools")
+class UFractureToolMergeSelected : public UFractureActionTool
+{
+public:
+	GENERATED_BODY()
+
+	UFractureToolMergeSelected(const FObjectInitializer& ObjInit) : Super(ObjInit) {}
+
+	// UFractureActionTool Interface
+	virtual FText GetDisplayText() const override;
+	virtual FText GetTooltipText() const override;
+	virtual FSlateIcon GetToolIcon() const override;
+	virtual void RegisterUICommand(FFractureEditorCommands* BindingContext) override;
+	virtual void Execute(TWeakPtr<FFractureEditorModeToolkit> InToolkit) override;
+};
+
+UCLASS(DisplayName = "Split Selected", Category = "FractureTools")
+class UFractureToolSplitSelected : public UFractureActionTool
+{
+public:
+	GENERATED_BODY()
+
+	UFractureToolSplitSelected(const FObjectInitializer& ObjInit) : Super(ObjInit) {}
+
+	// UFractureActionTool Interface
+	virtual FText GetDisplayText() const override;
+	virtual FText GetTooltipText() const override;
+	virtual FSlateIcon GetToolIcon() const override;
+	virtual void RegisterUICommand(FFractureEditorCommands* BindingContext) override;
+	virtual void Execute(TWeakPtr<FFractureEditorModeToolkit> InToolkit) override;
+};
+
+
 
 UCLASS(DisplayName = "Hide", Category = "FractureTools")
 class UFractureToolHide : public UFractureActionTool
diff --git a/Engine/Plugins/Experimental/PlanarCutPlugin/Source/PlanarCut/Private/GeometryMeshConversion.cpp b/Engine/Plugins/Experimental/PlanarCutPlugin/Source/PlanarCut/Private/GeometryMeshConversion.cpp
--- a/Engine/Plugins/Experimental/PlanarCutPlugin/Source/PlanarCut/Private/GeometryMeshConversion.cpp
+++ b/Engine/Plugins/Experimental/PlanarCutPlugin/Source/PlanarCut/Private/GeometryMeshConversion.cpp
@@ -2477,6 +2477,64 @@ int32 FDynamicMeshCollection::CutWithMultiplePlanes(
 }
 
 
+int32 FDynamicMeshCollection::SplitAllIslands(FGeometryCollection* Collection, double CollisionSampleSpacing)
+{
+	int32 FirstIdx = -1;
+	TArray<int32> GeometryForRemoval;
+
+	for (FMeshData& Surface : Meshes)
+	{
+		int32 SrcGeometryIdx = Collection->TransformToGeometryIndex[Surface.TransformIndex];
+		
+		int32 CreatedGeometryIdx = -1;
+		TArray<FDynamicMesh3> Islands;
+		if (SplitIslands(Surface.AugMesh, Islands))
+		{
+			for (int32 i = 0; i < Islands.Num(); i++)
+			{
+				FDynamicMesh3& Island = Islands[i];
+				FString BoneName = GetBoneName(*Collection, Surface.TransformIndex, i);
+				constexpr int32 InternalMaterialID = 0; // Note: there won't be new internal faces (to assign materials) from a split, so this value won't affect the output here
+				CreatedGeometryIdx = AppendToCollection(Surface.FromCollection, Island, CollisionSampleSpacing, Surface.TransformIndex, BoneName, *Collection, InternalMaterialID);
+
+				if (FirstIdx == -1)
+				{
+					FirstIdx = CreatedGeometryIdx;
+				}
+			}
+
+			GeometryForRemoval.Add(SrcGeometryIdx);
+		}
+	}
+
+	int32 NewFirstIdx = FirstIdx;
+
+	// remove or hide superfluous geometry
+	constexpr bool bRemoveOldGeometry = false; // if false, we just hide the geometry that we've replaced by fractured child geometry, rather than remove it
+	if (GeometryForRemoval.Num() > 0)
+	{
+		if (bRemoveOldGeometry)
+		{
+			GeometryForRemoval.Sort();
+			FManagedArrayCollection::FProcessingParameters ProcessingParams;
+#if !UE_BUILD_DEBUG
+			ProcessingParams.bDoValidation = false;
+#endif
+			Collection->RemoveElements(FGeometryCollection::GeometryGroup, GeometryForRemoval, ProcessingParams);
+			NewFirstIdx -= GeometryForRemoval.Num();
+		}
+		else
+		{
+			SetGeometryVisibility(Collection, GeometryForRemoval, false);
+		}
+	}
+
+
+	return NewFirstIdx;
+}
+
+
+
 int32 FDynamicMeshCollection::CutWithCellMeshes(const FInternalSurfaceMaterials& InternalSurfaceMaterials, const TArray<TPair<int32, int32>>& CellConnectivity, FCellMeshes& CellMeshes, FGeometryCollection* Collection, bool bSetDefaultInternalMaterialsFromCollection, double CollisionSampleSpacing)
 {
 	// TODO: should we do these cuts in parallel, and the appends sequentially below?
diff --git a/Engine/Plugins/Experimental/PlanarCutPlugin/Source/PlanarCut/Private/GeometryMeshConversion.h b/Engine/Plugins/Experimental/PlanarCutPlugin/Source/PlanarCut/Private/GeometryMeshConversion.h
--- a/Engine/Plugins/Experimental/PlanarCutPlugin/Source/PlanarCut/Private/GeometryMeshConversion.h
+++ b/Engine/Plugins/Experimental/PlanarCutPlugin/Source/PlanarCut/Private/GeometryMeshConversion.h
@@ -237,10 +237,21 @@ struct PLANARCUT_API FDynamicMeshCollection
 	 * @param CellsMeshes Meshed versions of the cells, with noise and material properties baked in
 	 * @param Collection Results will be stored in this
 	 * @param bSetDefaultInternalMaterialsFromCollection If true, set internal materials to the most common external material + 1, following a convenient artist convention
+	 * @param CollisionSampleSpacing If > 0, new geometry will have collision samples added (vertices not on any triangles) to fill any gaps greater than the this size
 	 * @return Index of the first created geometry
 	 */
 	int32 CutWithCellMeshes(const FInternalSurfaceMaterials& InternalSurfaceMaterials, const TArray<TPair<int32, int32>>& CellConnectivity, FCellMeshes& CellMeshes, FGeometryCollection* Collection, bool bSetDefaultInternalMaterialsFromCollection, double CollisionSampleSpacing);
 
+
+	/**
+	 * Split islands for all collection meshes, and append results to a geometry collection
+	 *
+	 * @param Collection Results will be stored in this
+	 * @param CollisionSampleSpacing If > 0, new geometry will have collision samples added (vertices not on any triangles) to fill any gaps greater than the this size
+	 * @return Index of the first created geometry, or -1 if nothing was split
+	 */
+	int32 SplitAllIslands(FGeometryCollection* Collection, double CollisionSampleSpacing);
+
 	static void SetVisibility(FGeometryCollection& Collection, int32 GeometryIdx, bool bVisible)
 	{
 		int32 FaceEnd = Collection.FaceCount[GeometryIdx] + Collection.FaceStart[GeometryIdx];
diff --git a/Engine/Plugins/Experimental/PlanarCutPlugin/Source/PlanarCut/Private/PlanarCut.cpp b/Engine/Plugins/Experimental/PlanarCutPlugin/Source/PlanarCut/Private/PlanarCut.cpp
--- a/Engine/Plugins/Experimental/PlanarCutPlugin/Source/PlanarCut/Private/PlanarCut.cpp
+++ b/Engine/Plugins/Experimental/PlanarCutPlugin/Source/PlanarCut/Private/PlanarCut.cpp
@@ -671,6 +671,41 @@ int32 CutMultipleWithPlanarCells(
 	return NewGeomStartIdx;
 }
 
+int32 SplitIslands(
+	FGeometryCollection& Collection,
+	const TArrayView<const int32>& TransformIndices,
+	double CollisionSampleSpacing,
+	FProgressCancel* Progress
+)
+{
+	FProgressCancel::FProgressScope CreateMeshCollectionScope = FProgressCancel::CreateScopeTo(Progress, .1);
+
+	FTransform CollectionToWorld = FTransform::Identity;
+
+	FDynamicMeshCollection MeshCollection(&Collection, TransformIndices, CollectionToWorld);
+	CreateMeshCollectionScope.Done();
+
+	if (Progress && Progress->Cancelled())
+	{
+		return -1;
+	}
+
+	FProgressCancel::FProgressScope SplitScope = FProgressCancel::CreateScopeTo(Progress, .99);
+	int32 NewGeomStartIdx = -1;
+	NewGeomStartIdx = MeshCollection.SplitAllIslands(&Collection, CollisionSampleSpacing);
+	SplitScope.Done();
+
+	if (Progress && Progress->Cancelled())
+	{
+		return -1;
+	}
+
+	FProgressCancel::FProgressScope ReindexScope = FProgressCancel::CreateScopeTo(Progress, 1);
+	Collection.ReindexMaterials();
+	ReindexScope.Done();
+	return NewGeomStartIdx;
+}
+
 FDynamicMesh3 ConvertMeshDescriptionToCuttingDynamicMesh(const FMeshDescription* CuttingMesh, int32 NumUVLayers, FProgressCancel* Progress)
 {
 	// populate the BaseMesh with a conversion of the input mesh.
@@ -1319,6 +1354,150 @@ int32 MergeBones(
 	return INDEX_NONE; // TODO: consider tracking smallest index of updated groups?  but no reason to do so currently
 }
 
+namespace
+{
+	void AddAllChildrenExceptEmbedded(FGeometryCollection& Collection, TSet<int32>& NodesSet, int32 Node)
+	{
+		if (Collection.SimulationType[Node] != FGeometryCollection::ESimulationTypes::FST_None)
+		{
+			NodesSet.Add(Node);
+			for (int32 Child : Collection.Children[Node])
+			{
+				AddAllChildrenExceptEmbedded(Collection, NodesSet, Child);
+			}
+		}
+	}
+}
+
+void MergeAllSelectedBones(
+	FGeometryCollection& Collection,
+	const TArrayView<const int32>& TransformIndices,
+	bool bUnionJoinedPieces
+)
+{
+	if (TransformIndices.IsEmpty())
+	{
+		return;
+	}
+
+	if (!Collection.HasAttribute("Level", FGeometryCollection::TransformGroup))
+	{
+		FGeometryCollectionClusteringUtility::UpdateHierarchyLevelOfChildren(&Collection, -1);
+	}
+	const TManagedArray<int32>& Level = Collection.GetAttribute<int32>("Level", FGeometryCollection::TransformGroup);
+	int32 TopLevelNode = TransformIndices[0];
+	int32 TopLevel = Level[TopLevelNode];
+
+	TSet<int32> AllNodesSet;
+	for (int32 Node : TransformIndices)
+	{
+		if (Level[Node] < TopLevel)
+		{
+			TopLevel = Level[Node];
+			TopLevelNode = Node;
+		}
+		AddAllChildrenExceptEmbedded(Collection, AllNodesSet, Node);
+	}
+	TArray<int32> AllNodes = AllNodesSet.Array();
+
+	if (AllNodes.Num() < 2)
+	{
+		return; // not enough pieces to need any merging
+	}
+
+	TArray<int32> AllUpdateIndices, AllRemoveIndices;
+	AllUpdateIndices.Add(TopLevelNode);
+	AllRemoveIndices.Reserve(AllNodes.Num() - 1);
+	for (int32 Node : AllNodes)
+	{
+		if (Node != TopLevelNode)
+		{
+			AllRemoveIndices.Add(Node);
+		}
+	}
+
+
+	AllRemoveIndices.Sort();
+	AllUpdateIndices.Sort();
+
+	FTransform CellsToWorld = FTransform::Identity;
+
+	FDynamicMeshCollection RemoveCollection(&Collection, AllRemoveIndices, CellsToWorld, true);
+	FDynamicMeshCollection UpdateCollection(&Collection, AllUpdateIndices, CellsToWorld, true);
+
+	using FMeshData = UE::PlanarCut::FDynamicMeshCollection::FMeshData;
+	bool bNeedAddGeometry = false;
+	if (UpdateCollection.Meshes.IsEmpty())
+	{
+		int32 NumUVLayers = Collection.NumUVLayers();
+		FMeshData* MeshData = new FMeshData(NumUVLayers);
+		MeshData->FromCollection = GeometryCollectionAlgo::GlobalMatrix(Collection.Transform, Collection.Parent, TopLevelNode);
+		MeshData->TransformIndex = TopLevelNode;
+		UpdateCollection.Meshes.Add(MeshData);
+
+		bNeedAddGeometry = true;
+	}
+	check(UpdateCollection.Meshes.Num() == 1);
+	FMeshData& UpMeshData = UpdateCollection.Meshes[0];
+	FDynamicMeshEditor MeshEditor(&UpMeshData.AugMesh);
+	for (int32 RmMeshIdx = 0; RmMeshIdx < RemoveCollection.Meshes.Num(); RmMeshIdx++)
+	{
+		FMeshData& RmMeshData = RemoveCollection.Meshes[RmMeshIdx];
+		if (bUnionJoinedPieces)
+		{
+			FMeshBoolean Boolean(&UpMeshData.AugMesh, &RmMeshData.AugMesh, &UpMeshData.AugMesh, FMeshBoolean::EBooleanOp::Union);
+			Boolean.bWeldSharedEdges = false;
+			Boolean.bSimplifyAlongNewEdges = true;
+			Boolean.Compute();
+		}
+		else
+		{
+			FMeshIndexMappings IndexMaps_Unused;
+			MeshEditor.AppendMesh(&RmMeshData.AugMesh, IndexMaps_Unused);
+		}
+	}
+
+	if (UpMeshData.AugMesh.TriangleCount() > 0)
+	{
+		if (bNeedAddGeometry)
+		{
+			if (Collection.TransformToGeometryIndex[TopLevelNode] == -1)
+			{
+				// Create a new geometry element to fill w/ the merged geometry
+				int32 GeometryIdx = Collection.AddElements(1, FGeometryCollection::GeometryGroup);
+				Collection.TransformToGeometryIndex[TopLevelNode] = GeometryIdx;
+				Collection.TransformIndex[GeometryIdx] = TopLevelNode;
+				Collection.FaceCount[GeometryIdx] = 0;
+				Collection.FaceStart[GeometryIdx] = Collection.Indices.Num();
+				Collection.VertexCount[GeometryIdx] = 0;
+				Collection.VertexStart[GeometryIdx] = Collection.Vertex.Num();
+			}
+		}
+		Collection.SimulationType[TopLevelNode] = FGeometryCollection::ESimulationTypes::FST_Rigid;
+		UpdateCollection.UpdateAllCollections(Collection);
+	}
+
+	TArray<int32> Children;
+	for (int32 RmTransformIdx : AllRemoveIndices)
+	{
+		for (int32 ChildIdx : Collection.Children[RmTransformIdx])
+		{
+			Children.Add(ChildIdx);
+		}
+	}
+	if (Children.Num())
+	{
+		GeometryCollectionAlgo::ParentTransforms(&Collection, TopLevelNode, Children);
+	}
+
+	// remove transforms for all geometry that was merged in
+	FManagedArrayCollection::FProcessingParameters ProcessingParams;
+#if !UE_BUILD_DEBUG
+	ProcessingParams.bDoValidation = false;
+#endif
+	Collection.RemoveElements(FGeometryCollection::TransformGroup, AllRemoveIndices, ProcessingParams);
+}
+
 
 void RecomputeNormalsAndTangents(bool bOnlyTangents, bool bMakeSharpEdges, float SharpAngleDegrees, FGeometryCollection& Collection, const TArrayView<const int32>& TransformIndices,
 								 bool bOnlyOddMaterials, const TArrayView<const int32>& WhichMaterials)
diff --git a/Engine/Plugins/Experimental/PlanarCutPlugin/Source/PlanarCut/Public/PlanarCut.h b/Engine/Plugins/Experimental/PlanarCutPlugin/Source/PlanarCut/Public/PlanarCut.h
--- a/Engine/Plugins/Experimental/PlanarCutPlugin/Source/PlanarCut/Public/PlanarCut.h
+++ b/Engine/Plugins/Experimental/PlanarCutPlugin/Source/PlanarCut/Public/PlanarCut.h
@@ -194,6 +194,22 @@ int32 PLANARCUT_API CutMultipleWithPlanarCells(
 	FProgressCancel* Progress = nullptr
 );
 
+/**
+ * Split the geometry at the given transforms into their connected components
+ *
+ * @param Collection		The collection to be cut
+ * @param TransformIndices	Which transform groups inside the collection to cut
+ * @param CollisionSampleSpacing	Target spacing between collision sample vertices
+ * @param Progress						Optionally tracks progress and supports early-cancel
+ * @return	index of first new geometry in the Output GeometryCollection, or -1 if no geometry was added
+ */
+int32 PLANARCUT_API SplitIslands(
+	FGeometryCollection& Collection,
+	const TArrayView<const int32>& TransformIndices,
+	double CollisionSampleSpacing,
+	FProgressCancel* Progress = nullptr
+);
+
 /**
  * Cut multiple Geometry groups inside a GeometryCollection with Planes, and add each cut cell back to the GeometryCollection as a new child of their source Geometry.  For geometries that would not be cut, nothing is added.
  *
@@ -305,6 +321,19 @@ int32 PLANARCUT_API MergeBones(
 	UE::PlanarCut::ENeighborSelectionMethod NeighborSelectionMethod
 );
 
+/**
+ * Merge all chosen nodes into a single node.  Unlike MergeBones(), does not account for proximity and will always only merge the selected bones.
+ *
+ * @param Collection				The collection to be processed
+ * @param TransformIndices			The transform indices to process, or empty if all should be processed
+ * @param bUnionJoinedPieces		Try to 'union' the merged pieces, removing internal triangles and connecting the shared cut boundary
+ */
+void PLANARCUT_API MergeAllSelectedBones(
+	FGeometryCollection& Collection,
+	const TArrayView<const int32>& TransformIndices,
+	bool bUnionJoinedPieces
+);
+
 /**
  * Recompute normals and tangents of selected geometry, optionally restricted to faces with odd or given material IDs (i.e. to target internal faces)
  *
