Propagate vertex colors internally on fracture via a smooth solve
 + Add a method to retrieve vertex connected components in a contiguous array, for faster per-component iteration

diff --git a/Engine/Plugins/Experimental/PlanarCutPlugin/Source/PlanarCut/PlanarCut.Build.cs b/Engine/Plugins/Experimental/PlanarCutPlugin/Source/PlanarCut/PlanarCut.Build.cs
--- a/Engine/Plugins/Experimental/PlanarCutPlugin/Source/PlanarCut/PlanarCut.Build.cs
+++ b/Engine/Plugins/Experimental/PlanarCutPlugin/Source/PlanarCut/PlanarCut.Build.cs
@@ -14,6 +14,7 @@ namespace UnrealBuildTool.Rules
 					"MeshUtilitiesCommon",
 					"StaticMeshDescription",
 					"MeshConversion",
+					"Eigen",
 				}
 				);
 			PublicDependencyModuleNames.AddRange(
diff --git a/Engine/Plugins/Experimental/PlanarCutPlugin/Source/PlanarCut/Private/GeometryMeshConversion.cpp b/Engine/Plugins/Experimental/PlanarCutPlugin/Source/PlanarCut/Private/GeometryMeshConversion.cpp
--- a/Engine/Plugins/Experimental/PlanarCutPlugin/Source/PlanarCut/Private/GeometryMeshConversion.cpp
+++ b/Engine/Plugins/Experimental/PlanarCutPlugin/Source/PlanarCut/Private/GeometryMeshConversion.cpp
@@ -9,6 +9,7 @@
 #include "Spatial/FastWinding.h"
 #include "Spatial/PointHashGrid3.h"
 #include "Spatial/MeshSpatialSort.h"
+#include "VertexConnectedComponents.h"
 #include "Util/IndexUtil.h"
 #include "Arrangement2d.h"
 #include "MeshAdapter.h"
@@ -44,6 +45,27 @@
 
 #include "Algo/Rotate.h"
 
+#if defined(_MSC_VER) && USING_CODE_ANALYSIS
+#pragma warning(push)
+#pragma warning(disable : 6011)
+#pragma warning(disable : 6387)
+#pragma warning(disable : 6313)
+#pragma warning(disable : 6294)
+#endif
+PRAGMA_DEFAULT_VISIBILITY_START
+THIRD_PARTY_INCLUDES_START
+#include <Eigen/Sparse>
+#include <Eigen/Core>
+#include <Eigen/SparseLU>
+#include <Eigen/OrderingMethods>
+#include <Eigen/Dense>
+THIRD_PARTY_INCLUDES_END
+PRAGMA_DEFAULT_VISIBILITY_END
+#if defined(_MSC_VER) && USING_CODE_ANALYSIS
+#pragma warning(pop)
+#endif
+#include <vector> // for Eigen sparse matrix construction
+
 using namespace UE::Geometry;
 
 
@@ -55,6 +77,11 @@ namespace PlanarCut
 // functions to setup geometry collection attributes on dynamic meshes
 namespace AugmentedDynamicMesh
 {
+	// An invalid color, to be replaced by neighboring valid colors (or DefaultVertexColor, if neighboring colors were not found)
+	// Use a large negative value as a clear unset / invalid value, but do not go all the way to -MaxReal (or overflow will break things)
+	const static FVector3f UnsetVertexColor = FVector3f(-FMathf::MaxReal * .25f, -FMathf::MaxReal * .25f, -FMathf::MaxReal * .25f);
+	const static FVector3f DefaultVertexColor = FVector3f::ZeroVector;
+
 	FName TangentUAttribName = "TangentUAttrib";
 	FName TangentVAttribName = "TangentVAttrib";
 	FName VisibleAttribName = "VisibleAttrib";
@@ -115,7 +142,7 @@ namespace AugmentedDynamicMesh
 
 	void Augment(FDynamicMesh3& Mesh, int32 NumUVChannels)
 	{
-		Mesh.EnableVertexColors(FVector3f(1, 1, 1));
+		Mesh.EnableVertexColors(UnsetVertexColor);
 		Mesh.EnableVertexNormals(FVector3f::UnitZ());
 		Mesh.EnableAttributes();
 		Mesh.Attributes()->EnableMaterialID();
@@ -677,7 +704,7 @@ namespace AugmentedDynamicMesh
 							if (VIDDist.Key == -1)
 							{
 								// no point within radius; can add a sample here
-								FVertexInfo Info(SamplePos, Normal);
+								FVertexInfo Info(SamplePos, Normal, AugmentedDynamicMesh::DefaultVertexColor);
 
 								int AddedVID = Mesh.AppendVertex(Info);
 								KnownSamples[ComponentIdx].InsertPointUnsafe(AddedVID, SamplePos);
@@ -1166,6 +1193,345 @@ namespace
 		Frame.ConstrainedAlignAxis(0, Target, Normal);
 		return Frame;
 	}
+
+
+	// Replace any unset vertex colors with colors from coincident or neighboring vertices (if bPropagateFromNeighbors) or DefaultVertexColor
+	void SetUnsetColors(FGeometryCollection* Collection, int32 FirstGeometryIndex, bool bPropagateFromNeighbors = true)
+	{
+		if (FirstGeometryIndex == INDEX_NONE)
+		{
+			return;
+		}
+
+		auto IsUnset = [](const FLinearColor& Color) -> bool
+		{
+			return Color.R < 0 || Color.G < 0 || Color.B < 0;
+		};
+		constexpr float CopyColorDistance = 1e-03;
+
+		int32 NumGeo = Collection->NumElements(FGeometryCollection::GeometryGroup);
+		
+		if (bPropagateFromNeighbors)
+		{
+			// Get vertex range for the geometry in and after FirstGeometryIndex
+			// Note: I think this will always just be the contiguous block of all vertices after the first geometry's start vertex,
+			//       but explicitly find the range to be safe (and to make the code easier to adapt)
+			int32 StartV = Collection->VertexStart[FirstGeometryIndex];
+			int32 EndV = Collection->VertexCount[FirstGeometryIndex] + StartV;
+			for (int32 GeoIdx = FirstGeometryIndex + 1; GeoIdx < NumGeo; ++GeoIdx)
+			{
+				int32 GeoStart = Collection->VertexStart[GeoIdx];
+				if (GeoStart < StartV)
+				{
+					StartV = GeoStart;
+				}
+				else
+				{
+					int32 GeoEnd = GeoStart + Collection->VertexCount[GeoIdx];
+					EndV = FMath::Max(GeoEnd, EndV);
+				}
+			}
+			int32 NumV = EndV - StartV;
+
+			// A generic representation of the full vertex graph
+			// (To be split into connected components and solved as a linear system per component)
+			struct FLink
+			{
+				int32 V[2]{ -1, -1 };
+				float Wt = 1;
+
+				FLink()
+				{}
+
+				FLink(int32 A, int32 B, float Wt) : V{ A, B }, Wt(Wt)
+				{}
+
+				static FLink Edge(int32 A, int32 B)
+				{
+					// Note: Currently just using constant weight (per half-edge)
+					constexpr float EdgeWt = .5;
+					return FLink(A, B, EdgeWt);
+				}
+			};
+			TArray<FLink> Links; // connections between vertices -- note: always symmetric, duplicates are allowed
+			Links.Reserve(NumV * 3);
+
+			TArray<FVector3f> FixedColors; // Note: FLinearColor has an alpha but FDynamicMesh3 currently only tracks rgb per vertex, so we only solve for rgb
+			TArray<float> FixedColorWts;
+			FixedColors.SetNumZeroed(NumV);
+			FixedColorWts.SetNumZeroed(NumV);
+
+			FVertexConnectedComponents Components(NumV); // Overall connected components (1 solve per component)
+			FVertexConnectedComponents VtxComps(NumV); // Vertex components (1 group of overlapping vertices per component)
+
+			for (int32 GeoIdx = FirstGeometryIndex; GeoIdx < NumGeo; ++GeoIdx)
+			{
+				int32 FStart = Collection->FaceStart[GeoIdx];
+				int32 FEnd = FStart + Collection->FaceCount[GeoIdx];
+				for (int32 FIdx = FStart; FIdx < FEnd; ++FIdx)
+				{
+					const FIntVector& Tri = Collection->Indices[FIdx];
+					// Note: This assumes triangles always have either fully set or fully unset vertices, so we only need to check one vertex
+					if (IsUnset(Collection->Color[Tri[0]]))
+					{
+						Components.ConnectVertices(Tri.X - StartV, Tri.Y - StartV);
+						Components.ConnectVertices(Tri.Y - StartV, Tri.Z - StartV);
+						Links.Add(FLink::Edge(Tri.X - StartV, Tri.Y - StartV));
+						Links.Add(FLink::Edge(Tri.Y - StartV, Tri.Z - StartV));
+						Links.Add(FLink::Edge(Tri.Z - StartV, Tri.X - StartV));
+					}
+				}
+			}
+
+			// put geometry in a shared coordinate space to spread color across geometry
+			TArray<FTransform> GlobalTransformArray;
+			GeometryCollectionAlgo::GlobalMatrices(Collection->Transform, Collection->Parent, GlobalTransformArray);
+			TArray<FVector3f> GlobalVertices;
+			GlobalVertices.SetNum(NumV);
+			for (int32 Idx = StartV; Idx < EndV; Idx++)
+			{
+				GlobalVertices[Idx - StartV] = (FVector3f)GlobalTransformArray[Collection->BoneMap[Idx]].TransformPosition(FVector(Collection->Vertex[Idx]));
+			}
+
+			// Create a hash grid of vertices with unset colors, and connect components for overlapping unset vertices
+			TPointHashGrid3f<int32> VertexHash(CopyColorDistance * 4.0f, INDEX_NONE);
+			TArray<int32> Neighbors;
+
+			for (int32 GeoIdx = FirstGeometryIndex; GeoIdx < NumGeo; ++GeoIdx)
+			{
+				int32 GeoStart = Collection->VertexStart[GeoIdx];
+				int32 GeoEnd = GeoStart + Collection->VertexCount[GeoIdx];
+				for (int32 Idx = GeoStart; Idx < GeoEnd; ++Idx)
+				{
+					FLinearColor Color = Collection->Color[Idx];
+					if (IsUnset(Color))
+					{
+						if (Components.GetComponentSize(Idx - StartV) == 1) // isolated vertex (not in any triangle), just leave as default color
+						{
+							Collection->Color[Idx] = FLinearColor(AugmentedDynamicMesh::DefaultVertexColor);
+							continue;
+						}
+						FVector3f Pt = GlobalVertices[Idx - StartV];
+
+						int32 SetID = Components.GetComponent(Idx - StartV);
+
+						Neighbors.Reset();
+						VertexHash.FindPointsInBall(Pt, CopyColorDistance, [&GlobalVertices, &Pt, StartV](int32 OtherIdx)
+							{
+								return DistanceSquared(Pt, GlobalVertices[OtherIdx - StartV]);
+							}, Neighbors);
+						for (int32 NbrIdx : Neighbors)
+						{
+							VtxComps.ConnectVertices(Idx - StartV, NbrIdx - StartV);
+							Components.ConnectVertices(SetID, NbrIdx - StartV);
+						}
+						VertexHash.InsertPointUnsafe(Idx, Pt);
+					}
+				}
+			}
+
+			TSet<int32> CanSolveComponents;
+
+			// Fix colors on vertices that are coincident to set colors
+			for (int32 GeoIdx = FirstGeometryIndex; GeoIdx < NumGeo; ++GeoIdx)
+			{
+				int32 GeoStart = Collection->VertexStart[GeoIdx];
+				int32 GeoEnd = GeoStart + Collection->VertexCount[GeoIdx];
+				for (int32 Idx = GeoStart; Idx < GeoEnd; ++Idx)
+				{
+					FLinearColor Color = Collection->Color[Idx];
+					if (!IsUnset(Color))
+					{
+						FVector3f Pt = GlobalVertices[Idx - StartV];
+						Neighbors.Reset();
+						VertexHash.FindPointsInBall(Pt, CopyColorDistance, [&GlobalVertices, &Pt, StartV](int32 OtherIdx)
+							{
+								return DistanceSquared(Pt, GlobalVertices[OtherIdx - StartV]);
+							}, Neighbors);
+						for (int32 NbrIdx : Neighbors)
+						{
+							int32 Component = Components.GetComponent(NbrIdx - StartV);
+							int32 ComponentSize = Components.GetComponentSize(NbrIdx - StartV);
+							if (ComponentSize == 1)
+							{
+								Collection->Color[NbrIdx] = Color; // Isolated vertex, no solve needed
+							}
+							CanSolveComponents.Add(Component);
+							int32 SharedVtx = VtxComps.GetComponent(NbrIdx - StartV);
+							FixedColors[SharedVtx] += FVector3f(Color.R, Color.G, Color.B);
+							FixedColorWts[SharedVtx] += 1.0f;
+						}
+					}
+				}
+			}
+
+			// Average color at any vertex w/ a fixed color
+			for (int32 ColorIdx = 0; ColorIdx < FixedColors.Num(); ++ColorIdx)
+			{
+				float& Wt = FixedColorWts[ColorIdx];
+				if (Wt > 0.0f)
+				{
+					FixedColors[ColorIdx] /= Wt;
+					Wt = 1.0f;
+				}
+			}
+
+			TArray<int32> ToComponentMap;
+			ToComponentMap.Reserve(NumV);
+			TArray<float> DiagonalWts;
+
+			TArray<int32> ContigComponentVertices = Components.MakeContiguousComponentsArray(NumV);
+
+			// Solve for vertex colors per component
+			for (int32 ContigStart = 0, NextStart = -1; ContigStart < NumV; ContigStart = NextStart)
+			{
+				int32 ComponentID = Components.GetComponent(ContigComponentVertices[ContigStart]);
+				int32 ComponentSize = Components.GetComponentSize(ComponentID);
+				NextStart = ContigStart + ComponentSize;
+
+				if (ComponentSize == 1)
+				{
+					continue;
+				}
+				if (!CanSolveComponents.Contains(ComponentID))
+				{
+					continue;
+				}
+
+				using FSparseMatf = Eigen::SparseMatrix<float, Eigen::ColMajor>;
+				using FMatrixTripletf = Eigen::Triplet<float>;
+				std::vector<FMatrixTripletf> EntryTriplets;
+
+				// Make an index map for this component's vertices
+				ToComponentMap.Init(-1, NumV);
+				int32 NumToSolve = 0;
+				for (int32 ContigIdx = ContigStart; ContigIdx < NextStart; ++ContigIdx)
+				{
+					int32 LocalIdx = ContigComponentVertices[ContigIdx];
+					int32 GlobalIdx = LocalIdx + StartV;
+					int32 SharedIdx = VtxComps.GetComponent(LocalIdx);
+					if (FixedColorWts[SharedIdx] > 0)
+					{
+						FLinearColor Color(FixedColors[SharedIdx]);
+						Collection->Color[GlobalIdx] = Color; // copy fixed color out
+						continue; // has fixed color, do not include in solve
+					}
+					if (ToComponentMap[SharedIdx] != -1)
+					{
+						continue; // already mapped
+					}
+
+					ToComponentMap[SharedIdx] = NumToSolve;
+
+					++NumToSolve;
+				}
+
+				if (NumToSolve == 0)
+				{
+					continue;
+				}
+
+				FSparseMatf SparseMatrix(NumToSolve, NumToSolve);
+				Eigen::VectorXf X[3]{ Eigen::VectorXf(NumToSolve), Eigen::VectorXf(NumToSolve), Eigen::VectorXf(NumToSolve) };
+				Eigen::VectorXf B[3]{ Eigen::VectorXf(NumToSolve), Eigen::VectorXf(NumToSolve), Eigen::VectorXf(NumToSolve) };
+				for (int32 SubIdx = 0; SubIdx < 3; ++SubIdx)
+				{
+					B[SubIdx].setZero();
+				}
+				DiagonalWts.Reset(NumToSolve);
+				DiagonalWts.SetNumZeroed(NumToSolve, false);
+
+				// Build the sparse matrix and rhs for the component
+				for (const FLink& Link : Links)
+				{
+					int32 SharedA = VtxComps.GetComponent(Link.V[0]);
+					int32 SharedB = VtxComps.GetComponent(Link.V[1]);
+					int32 LocalA = ToComponentMap[SharedA];
+					int32 LocalB = ToComponentMap[SharedB];
+					if (LocalA == INDEX_NONE)
+					{
+						if (LocalB == INDEX_NONE)
+						{
+							continue;
+						}
+						Swap(SharedA, SharedB);
+						Swap(LocalA, LocalB);
+					}
+					if (LocalB == INDEX_NONE)
+					{
+						if (FixedColorWts[SharedB] > 0)
+						{
+							FVector3f BVal = FixedColors[SharedB] * Link.Wt;
+							for (int32 SubIdx = 0; SubIdx < 3; ++SubIdx)
+							{
+								B[SubIdx][LocalA] += BVal[SubIdx];
+							}
+							DiagonalWts[LocalA] += Link.Wt;
+						}
+					}
+					else
+					{
+						EntryTriplets.push_back(FMatrixTripletf(LocalA, LocalB, -Link.Wt));
+						EntryTriplets.push_back(FMatrixTripletf(LocalB, LocalA, -Link.Wt));
+						DiagonalWts[LocalA] += Link.Wt;
+						DiagonalWts[LocalB] += Link.Wt;
+					}
+				}
+				for (int32 Idx = 0; Idx < NumToSolve; ++Idx)
+				{
+					EntryTriplets.push_back(FMatrixTripletf(Idx, Idx, DiagonalWts[Idx]));
+				}
+
+				// Solve linear system for internal colors
+				SparseMatrix.setFromTriplets(EntryTriplets.begin(), EntryTriplets.end());
+				SparseMatrix.makeCompressed();
+				
+				Eigen::SparseLU<FSparseMatf, Eigen::COLAMDOrdering<int>> MatrixSolver;
+				MatrixSolver.isSymmetric(true);
+				MatrixSolver.analyzePattern(SparseMatrix);
+				MatrixSolver.factorize(SparseMatrix);
+
+				ParallelFor(3, [&](int32 Idx)
+					{
+						X[Idx] = MatrixSolver.solve(B[Idx]);
+					});
+
+				// Copy solved colors out
+				for (int32 ContigIdx = ContigStart; ContigIdx < NextStart; ++ContigIdx)
+				{
+					int32 LocalIdx = ContigComponentVertices[ContigIdx];
+					int32 GlobalIdx = LocalIdx + StartV;
+					int32 SharedIdx = VtxComps.GetComponent(LocalIdx);
+					int32 CompIdx = ToComponentMap[SharedIdx];
+					if (CompIdx != INDEX_NONE)
+					{
+						FVector3f SolvedColor;
+						for (int32 SubIdx = 0; SubIdx < 3; ++SubIdx)
+						{
+							// Note: make sure the solved color is non-negative, so solver error cannot make the color 'unset'
+							SolvedColor[SubIdx] = FMath::Max(X[SubIdx][CompIdx], 0.f);
+						}
+
+						Collection->Color[GlobalIdx] = FLinearColor(SolvedColor);
+					}
+				}
+			}
+		}
+
+		// Replace unset color with default color
+		for (int32 GeoIdx = FirstGeometryIndex; GeoIdx < NumGeo; ++GeoIdx)
+		{
+			int32 VStart = Collection->VertexStart[GeoIdx];
+			int32 VEnd = VStart + Collection->VertexCount[GeoIdx];
+			for (int32 Idx = VStart; Idx < VEnd; ++Idx)
+			{
+				if (IsUnset(Collection->Color[Idx]))
+				{
+					Collection->Color[Idx] = FLinearColor(AugmentedDynamicMesh::DefaultVertexColor);
+				}
+			}
+		}
+	}
 }
 
 
@@ -1199,7 +1565,11 @@ FCellMeshes::FCellMeshes(int32 NumUVLayersIn, const FDynamicMesh3& SingleCutter,
 		AugmentedDynamicMesh::Augment(CellMeshes[0].AugMesh, NumUVLayers);
 	}
 
-	
+	// Make sure color is unset
+	for (int VID : CellMeshes[0].AugMesh.VertexIndicesItr())
+	{
+		CellMeshes[0].AugMesh.SetVertexColor(VID, AugmentedDynamicMesh::UnsetVertexColor);
+	}
 
 	// first mesh is the same as the second mesh, but will be subtracted b/c it's the "outside cell"
 	// TODO: special case this logic so we don't have to hold two copies of the exact same mesh!
@@ -1498,7 +1868,7 @@ void FCellMeshes::CreateMeshesForBoundedPlanesWithoutNoise(int NumCells, const F
 		PlaneVertInfo.bHaveC = true;
 		PlaneVertInfo.bHaveUV = false;
 		PlaneVertInfo.bHaveN = true;
-		PlaneVertInfo.Color = FVector3f(1, 1, 1);
+		PlaneVertInfo.Color = AugmentedDynamicMesh::UnsetVertexColor;
 		int VertStart[2]{ -1, -1 };
 		for (int MeshIdx = 0; MeshIdx < NumMeshes; MeshIdx++)
 		{
@@ -1611,7 +1981,7 @@ void FCellMeshes::CreateMeshesForBoundedPlanesWithNoise(int NumCells, const FPla
 	{
 		AugmentedDynamicMesh::EnableUVChannels(PlaneMesh, NumUVLayers);
 		PlaneMesh.EnableVertexNormals(FVector3f::UnitZ());
-		PlaneMesh.EnableVertexColors(FVector3f(1, 1, 1));
+		PlaneMesh.EnableVertexColors(AugmentedDynamicMesh::UnsetVertexColor);
 		PlaneMesh.EnableAttributes();
 		PlaneMesh.Attributes()->EnableMaterialID();
 		PlaneMesh.Attributes()->AttachAttribute(OriginalPositionAttribute, new TDynamicMeshVertexAttribute<double, 3>(&PlaneMesh));
@@ -1668,7 +2038,7 @@ void FCellMeshes::CreateMeshesForBoundedPlanesWithNoise(int NumCells, const FPla
 			PlaneVertInfo.bHaveN = true;
 			PlaneVertInfo.Normal = Normal;
 			// UVs will be set below, after noise is added
-			PlaneVertInfo.Color = FVector3f(1, 1, 1);
+			PlaneVertInfo.Color = AugmentedDynamicMesh::UnsetVertexColor;
 
 			FPolygon2f Polygon;
 			for (int BoundaryVertex : PlaneBoundary)
@@ -1980,7 +2350,7 @@ void FCellMeshes::CreateMeshesForSinglePlane(const FPlanarCells& Cells, const FA
 	FVertexInfo PlaneVertInfo;
 	PlaneVertInfo.bHaveC = true;
 	PlaneVertInfo.bHaveN = true;
-	PlaneVertInfo.Color = FVector3f(1, 1, 1);
+	PlaneVertInfo.Color = AugmentedDynamicMesh::UnsetVertexColor;
 	PlaneVertInfo.Normal = -FVector3f(Plane.GetNormal());
 
 	for (int CornerIdx = 0; CornerIdx < 4; CornerIdx++)
@@ -2452,6 +2822,8 @@ int32 FDynamicMeshCollection::CutWithMultiplePlanes(
 		return INDEX_NONE;
 	}
 
+	SetUnsetColors(Collection, FirstCreatedIndex);
+
 	int32 NewFirstIndex = FirstCreatedIndex;
 
 	constexpr bool bRemoveOldGeometry = false; // if false, we just hide the geometry that we've replaced by fractured child geometry, rather than remove it
@@ -2660,6 +3032,8 @@ int32 FDynamicMeshCollection::CutWithCellMeshes(const FInternalSurfaceMaterials&
 		}
 	}
 
+	SetUnsetColors(Collection, FirstIdx);
+
 	int32 NewFirstIdx = FirstIdx;
 
 	// remove or hide superfluous geometry
@@ -2840,6 +3214,8 @@ bool FDynamicMeshCollection::UpdateAllCollections(FGeometryCollection& Collectio
 		bAllSucceeded &= bSucceeded;
 	}
 
+	SetUnsetColors(&Collection, 0, false);
+
 	return bAllSucceeded;
 }
 
diff --git a/Engine/Source/Runtime/GeometryCore/Public/VertexConnectedComponents.h b/Engine/Source/Runtime/GeometryCore/Public/VertexConnectedComponents.h
--- a/Engine/Source/Runtime/GeometryCore/Public/VertexConnectedComponents.h
+++ b/Engine/Source/Runtime/GeometryCore/Public/VertexConnectedComponents.h
@@ -303,6 +303,39 @@ public:
 		return ComponentMap;
 	}
 
+	// Return an ordering of the vertex indices so that each connected component is in a contiguous block
+	TArray<int32> MakeContiguousComponentsArray(int32 MaxVID)
+	{
+		TMap<int32, int32> ComponentLoc;
+		TArray<int32> Contiguous;
+		Contiguous.SetNum(MaxVID);
+		int32 LastSingleEntryIdx = MaxVID;
+		int32 FirstUnusedIdx = 0;
+		for (int32 VID = 0; VID < MaxVID; VID++)
+		{
+			int32 SetID = DisjointSet.Find(VID);
+			int32 SetSize = DisjointSet.Sizes[SetID];
+			if (SetSize == 1) // just place the single-element groups at the end, no need to track in map
+			{
+				LastSingleEntryIdx--;
+				Contiguous[LastSingleEntryIdx] = VID;
+				continue;
+			}
+			int32* Loc = ComponentLoc.Find(SetID);
+			if (Loc)
+			{
+				Contiguous[(*Loc)++] = VID;
+			}
+			else
+			{
+				Contiguous[FirstUnusedIdx] = VID;
+				ComponentLoc.Add(SetID, FirstUnusedIdx + 1);
+				FirstUnusedIdx += SetSize;
+			}
+		}
+		return Contiguous;
+	}
+
 	inline int32 GetComponent(int32 VertexID)
 	{
 		return DisjointSet.Find(VertexID);
