fix fractureautouv's box projection messing up unrelated UVs and breaking for isolated triangles

diff --git a/Engine/Plugins/Experimental/PlanarCutPlugin/Source/PlanarCut/Private/FractureAutoUV.cpp b/Engine/Plugins/Experimental/PlanarCutPlugin/Source/PlanarCut/Private/FractureAutoUV.cpp
--- a/Engine/Plugins/Experimental/PlanarCutPlugin/Source/PlanarCut/Private/FractureAutoUV.cpp
+++ b/Engine/Plugins/Experimental/PlanarCutPlugin/Source/PlanarCut/Private/FractureAutoUV.cpp
@@ -419,15 +419,13 @@ bool BoxProjectUVs(
 				TargetTris.Add(TID);
 			}
 		}
-		if (TargetTris.Num() == 0)
+		if (TargetTris.Num() > 0)
 		{
-			continue;
+			FDynamicMeshUVEditor UVEd(&Mesh, TargetUVLayer, false);
+			FFrame3d BoxFrame; // defaults to origin / no rotation
+			UVEd.SetTriangleUVsFromBoxProjection(TargetTris, [](const FVector3d& Pos) { return Pos; }, BoxFrame, BoxDimensions, 1);
 		}
 
-		FDynamicMeshUVEditor UVEd(&Mesh, TargetUVLayer, false);
-		FFrame3d BoxFrame; // defaults to origin / no rotation
-		UVEd.SetTriangleUVsFromBoxProjection(TargetTris, [](const FVector3d& Pos) { return Pos; }, BoxFrame, BoxDimensions);
-
 		Mesh.CompactInPlace();
 
 		// Re-split vertices with different UVs/normals/tangents and transfer attributes back (required because we weld edges above)
