Make fracture brick cutter work correctly with zero grout, rather than forcing the grout to always be non-zero

diff --git a/Engine/Plugins/Experimental/ChaosEditor/Source/FractureEditor/Private/FractureToolBrick.cpp b/Engine/Plugins/Experimental/ChaosEditor/Source/FractureEditor/Private/FractureToolBrick.cpp
--- a/Engine/Plugins/Experimental/ChaosEditor/Source/FractureEditor/Private/FractureToolBrick.cpp
+++ b/Engine/Plugins/Experimental/ChaosEditor/Source/FractureEditor/Private/FractureToolBrick.cpp
@@ -425,10 +425,9 @@ int32 UFractureToolBrick::ExecuteFracture(const FFractureToolContext& FractureCo
 
 		TArray<FBox> BricksToCut;
 
-		// space the bricks by the grout setting, constrained to not erase the bricks or have zero grout
-		// (currently zero grout bricks would break assumptions in the fracture)
+		// space the bricks by the grout setting, constrained to not erase the bricks
 		const float MinDim = FMath::Min3(BrickHalfDimensions.X, BrickHalfDimensions.Y, BrickHalfDimensions.Z);
-		const float HalfGrout = FMath::Clamp(0.5f * CutterSettings->Grout, MinDim * 0.02f, MinDim * 0.98f);
+		const float HalfGrout = FMath::Clamp(0.5f * CutterSettings->Grout, 0, MinDim * 0.98f);
 		const FVector HalfBrick(BrickHalfDimensions - HalfGrout);
 		const FBox BrickBox(-HalfBrick, HalfBrick);
 
@@ -441,7 +440,8 @@ int32 UFractureToolBrick::ExecuteFracture(const FFractureToolContext& FractureCo
 		BrickOp->Selection = FractureContext.GetSelection();
 		BrickOp->Grout = 0; // CutterSettings->Grout; // Note: Grout is currently baked directly into the brick cells above
 		BrickOp->PointSpacing = CollisionSettings->GetPointSpacing();
-		BrickOp->Cells = FPlanarCells(BricksToCut);
+		const bool bBricksAreTouching = CutterSettings->Grout <= UE_KINDA_SMALL_NUMBER;
+		BrickOp->Cells = FPlanarCells(BricksToCut, bBricksAreTouching);
 		if (CutterSettings->Amplitude > 0.0f)
 		{
 			FNoiseSettings Settings;
diff --git a/Engine/Plugins/Experimental/PlanarCutPlugin/Source/PlanarCut/Private/PlanarCut.cpp b/Engine/Plugins/Experimental/PlanarCutPlugin/Source/PlanarCut/Private/PlanarCut.cpp
--- a/Engine/Plugins/Experimental/PlanarCutPlugin/Source/PlanarCut/Private/PlanarCut.cpp
+++ b/Engine/Plugins/Experimental/PlanarCutPlugin/Source/PlanarCut/Private/PlanarCut.cpp
@@ -6,6 +6,7 @@
 #include "Spatial/FastWinding.h"
 #include "Spatial/PointHashGrid3.h"
 #include "Spatial/MeshSpatialSort.h"
+#include "Spatial/SparseDynamicOctree3.h"
 #include "Util/IndexUtil.h"
 #include "Arrangement2d.h"
 #include "MeshAdapter.h"
@@ -194,36 +195,206 @@ FPlanarCells::FPlanarCells(const TArrayView<const FVector> Sites, FVoronoiDiagra
 	}
 }
 
-FPlanarCells::FPlanarCells(const TArrayView<const FBox> Boxes)
+FPlanarCells::FPlanarCells(const TArrayView<const FBox> Boxes, bool bResolveAdjacencies)
 {
 	AssumeConvexCells = true;
 	NumCells = Boxes.Num();
 	TArray<FBox> BoxesCopy(Boxes);
 	
+	if (!bResolveAdjacencies) // if the boxes aren't touching, we can just make a completely independent cell per box
+	{
+		for (int32 BoxIdx = 0; BoxIdx < NumCells; BoxIdx++)
+		{
+			const FBox& Box = Boxes[BoxIdx];
+			const FVector& Min = Box.Min;
+			const FVector& Max = Box.Max;
+
+			int32 VIdx = PlaneBoundaryVertices.Num();
+			PlaneBoundaryVertices.Add(Min);
+			PlaneBoundaryVertices.Add(FVector(Max.X, Min.Y, Min.Z));
+			PlaneBoundaryVertices.Add(FVector(Max.X, Max.Y, Min.Z));
+			PlaneBoundaryVertices.Add(FVector(Min.X, Max.Y, Min.Z));
+
+			PlaneBoundaryVertices.Add(FVector(Min.X, Min.Y, Max.Z));
+			PlaneBoundaryVertices.Add(FVector(Max.X, Min.Y, Max.Z));
+			PlaneBoundaryVertices.Add(Max);
+			PlaneBoundaryVertices.Add(FVector(Min.X, Max.Y, Max.Z));
+
+			AddPlane(FPlane(FVector(0, 0, -1), -Min.Z), BoxIdx, -1, { VIdx + 0, VIdx + 1, VIdx + 2, VIdx + 3 });
+			AddPlane(FPlane(FVector(0, 0, 1), Max.Z), BoxIdx, -1, { VIdx + 4, VIdx + 7, VIdx + 6, VIdx + 5 });
+			AddPlane(FPlane(FVector(0, -1, 0), -Min.Y), BoxIdx, -1, { VIdx + 0, VIdx + 4, VIdx + 5, VIdx + 1 });
+			AddPlane(FPlane(FVector(0, 1, 0), Max.Y), BoxIdx, -1, { VIdx + 3, VIdx + 2, VIdx + 6, VIdx + 7 });
+			AddPlane(FPlane(FVector(-1, 0, 0), -Min.X), BoxIdx, -1, { VIdx + 0, VIdx + 3, VIdx + 7, VIdx + 4 });
+			AddPlane(FPlane(FVector(1, 0, 0), Max.X), BoxIdx, -1, { VIdx + 1, VIdx + 5, VIdx + 6, VIdx + 2 });
+		}
+
+		return;
+	}
+
+	// If boxes might be touching, we need to subdivide each box along each overlap,
+	// and make sure to only construct one copy any vertex/face that is shared between multiple boxes
+
+	// Build an octree of boxes to allow us to quickly find neighbors, and determine how to subdivide the boxes
+	FSparseDynamicOctree3 BoxTree;
+	double BoxMaxDim = 0;
+	for (const FBox& Box : Boxes)
+	{
+		BoxMaxDim = FMath::Max(BoxMaxDim, Box.GetSize().GetMax());
+	}
+	BoxTree.RootDimension = BoxMaxDim * 4;
+	for (int32 BoxIdx = 0; BoxIdx < NumCells; BoxIdx++)
+	{
+		BoxTree.InsertObject(BoxIdx, Boxes[BoxIdx]);
+	}
+
+	// Build a hash grid for vertices, to share vertices across boxes
+	constexpr double AddVertTolerance = UE_DOUBLE_KINDA_SMALL_NUMBER;
+	TPointHashGrid3d<int32> PosToVert(AddVertTolerance * 10, INDEX_NONE);
+	TMap<FIntVector4, int32> VertsToPlane;
+
+	// Track overlaps between neighboring boxes in each dimension separately; each overlap is a new subdivision on that axis
+	TArray<int32> OverlapIndices;
+	TArray<double> Overlaps[3];
 	for (int32 BoxIdx = 0; BoxIdx < NumCells; BoxIdx++)
 	{
-		const FBox &Box = Boxes[BoxIdx];
-		const FVector &Min = Box.Min;
-		const FVector &Max = Box.Max;
-
-		int32 VIdx = PlaneBoundaryVertices.Num();
-		PlaneBoundaryVertices.Add(Min);
-		PlaneBoundaryVertices.Add(FVector(Max.X, Min.Y, Min.Z));
-		PlaneBoundaryVertices.Add(FVector(Max.X, Max.Y, Min.Z));
-		PlaneBoundaryVertices.Add(FVector(Min.X, Max.Y, Min.Z));
-
-		PlaneBoundaryVertices.Add(FVector(Min.X, Min.Y, Max.Z));
-		PlaneBoundaryVertices.Add(FVector(Max.X, Min.Y, Max.Z));
-		PlaneBoundaryVertices.Add(Max);
-		PlaneBoundaryVertices.Add(FVector(Min.X, Max.Y, Max.Z));
-
-		AddPlane(FPlane(FVector(0, 0, -1), -Min.Z), BoxIdx, -1, { VIdx + 0, VIdx + 1, VIdx + 2, VIdx + 3 });
-		AddPlane(FPlane(FVector(0, 0, 1),	Max.Z), BoxIdx, -1, { VIdx + 4, VIdx + 7, VIdx + 6, VIdx + 5 });
-		AddPlane(FPlane(FVector(0, -1, 0), -Min.Y), BoxIdx, -1, { VIdx + 0, VIdx + 4, VIdx + 5, VIdx + 1 });
-		AddPlane(FPlane(FVector(0, 1, 0),	Max.Y), BoxIdx, -1, { VIdx + 3, VIdx + 2, VIdx + 6, VIdx + 7 });
-		AddPlane(FPlane(FVector(-1, 0, 0), -Min.X), BoxIdx, -1, { VIdx + 0, VIdx + 3, VIdx + 7, VIdx + 4 });
-		AddPlane(FPlane(FVector(1, 0, 0),	Max.X), BoxIdx, -1, { VIdx + 1, VIdx + 5, VIdx + 6, VIdx + 2 });
+		const FBox& Box = Boxes[BoxIdx];
+		const FVector& Min = Box.Min;
+		const FVector& Max = Box.Max;
+
+		Overlaps[0].Reset();
+		Overlaps[1].Reset();
+		Overlaps[2].Reset();
+		OverlapIndices.Reset();
+		
+		constexpr double QueryExpandTolerance = AddVertTolerance;
+		
+		FBox ExpandedBox = Box.ExpandBy(QueryExpandTolerance);
+		BoxTree.RangeQuery(ExpandedBox, OverlapIndices);
+		for (int32 OverlapIdx : OverlapIndices)
+		{
+			// skip if same box or not an actual overlap
+			const FBox& OtherBox = Boxes[OverlapIdx];
+			if (BoxIdx == OverlapIdx || !ExpandedBox.Intersect(OtherBox))
+			{
+				continue;
+			}
+
+			auto AddIfInRange = [AddVertTolerance](double Val, double RangeMin, double RangeMax, TArray<double>& AddTo)
+			{
+				if (Val > RangeMin + AddVertTolerance && Val < RangeMax - AddVertTolerance)
+				{
+					AddTo.Add(Val);
+				}
+			};
+
+			for (int32 Dim = 0; Dim < 3; ++Dim)
+			{
+				AddIfInRange(OtherBox.Min[Dim], Box.Min[Dim], Box.Max[Dim], Overlaps[Dim]);
+				AddIfInRange(OtherBox.Max[Dim], Box.Min[Dim], Box.Max[Dim], Overlaps[Dim]);
+			}
+		}
+
+		for (int32 Dim = 0; Dim < 3; ++Dim)
+		{
+			Overlaps[Dim].Add(Box.Min[Dim]);
+			Overlaps[Dim].Add(Box.Max[Dim]);
+			Overlaps[Dim].Sort();
+
+			double Last = Overlaps[Dim][0];
+			int32 FillIdx = 1;
+			for (int32 Idx = 1; Idx < Overlaps[Dim].Num(); ++Idx)
+			{
+				if (Overlaps[Dim][Idx] - Last >= AddVertTolerance)
+				{
+					Overlaps[Dim][FillIdx] = Overlaps[Dim][Idx];
+					Last = Overlaps[Dim][Idx];
+					++FillIdx;
+				}
+				// else overlap was too close to 'last', so we skip it
+			}
+			Overlaps[Dim].SetNum(FillIdx, false);
+			if (Overlaps[Dim].Num() == 1)
+			{
+				Overlaps[Dim].Add(Box.Max[Dim]);
+			}
+			else
+			{
+				// make sure to always end at the exact max
+				Overlaps[Dim].Last() = Box.Max[Dim];
+			}
+		}
+
+		// Helper to get a shared vertex, first by looking if we've already made it locally, then by checking the global hash grid
+		TMap<FIndex3i, int32> GridToVert; // A local map for all the vertices we've already found on the box
+		auto GetVertex = [this, &Overlaps, &GridToVert, &PosToVert, AddVertTolerance](int32 FaceDim, int32 UDim, int32 VDim, int32 FaceCoord, int32 U, int32 V)
+		{
+			FIndex3i Coord;
+			Coord[FaceDim] = FaceCoord;
+			Coord[UDim] = U;
+			Coord[VDim] = V;
+			int32* FoundV = GridToVert.Find(Coord);
+			if (FoundV)
+			{
+				return *FoundV;
+			}
+			FVector Pos(Overlaps[0][Coord.A], Overlaps[1][Coord.B], Overlaps[2][Coord.C]);
+			TPair<int32, double> Found = PosToVert.FindNearestInRadius(Pos, AddVertTolerance, [&](const int32& Idx) {return FVector::DistSquared(Pos, PlaneBoundaryVertices[Idx]);});
+			if (Found.Key != INDEX_NONE)
+			{
+				return Found.Key;
+			}
+
+			int32 NewV = PlaneBoundaryVertices.Add(Pos);
+			GridToVert.Add(Coord, NewV);
+			PosToVert.InsertPointUnsafe(NewV, Pos);
+			return NewV;
+		};
+
+		// Now that we have decided how to subdivide each dimension (w/ the Overlaps arrays), we need to construct each face (or link to the face, if it already exists)
+		for (int32 FaceDim = 0; FaceDim < 3; ++FaceDim)
+		{
+			for (int32 Side = 0; Side < 2; ++Side)
+			{
+				int32 UDim = (FaceDim + 1) % 3;
+				int32 VDim = (FaceDim + 2) % 3;
+				int32 FaceCoord = (Side == 1) ? (Overlaps[FaceDim].Num() - 1) : 0;
+				double FaceSign = double(Side * 2 - 1);
+				FVector Normal(0, 0, 0);
+				Normal[FaceDim] = FaceSign;
+				double PlaneW = FaceSign * (Side == 0 ? Min[FaceDim] : Max[FaceDim]);
+				for (int32 U = 0; U + 1 < Overlaps[UDim].Num(); ++U)
+				{
+					for (int32 V = 0; V + 1 < Overlaps[VDim].Num(); ++V)
+					{
+						int32 V00 = GetVertex(FaceDim, UDim, VDim, FaceCoord, U, V);
+						int32 V10 = GetVertex(FaceDim, UDim, VDim, FaceCoord, U + 1, V);
+						int32 V11 = GetVertex(FaceDim, UDim, VDim, FaceCoord, U + 1, V + 1);
+						int32 V01 = GetVertex(FaceDim, UDim, VDim, FaceCoord, U, V + 1);
+
+						FIntVector4 FaceKey(V00, V10, V11, V01);
+						int32* FoundPlane = VertsToPlane.Find(FaceKey);
+						if (FoundPlane)
+						{
+							PlaneCells[*FoundPlane].Value = BoxIdx;
+							continue;
+						}
+
+						TArray<int32> BoundaryVerts;
+						if (Side == 0)
+						{
+							BoundaryVerts = { V00, V10, V11, V01 };
+						}
+						else // reverse winding for the 'far' face
+						{
+							BoundaryVerts = { V10, V00, V01, V11 };
+						}
+						VertsToPlane.Add(FaceKey, AddPlane(FPlane(Normal, PlaneW), BoxIdx, -1, BoundaryVerts));
+					}
+				}
+			}
+		}
 	}
+
 }
 
 FPlanarCells::FPlanarCells(const FBox& Region, const FIntVector& CubesPerAxis)
diff --git a/Engine/Plugins/Experimental/PlanarCutPlugin/Source/PlanarCut/Public/PlanarCut.h b/Engine/Plugins/Experimental/PlanarCutPlugin/Source/PlanarCut/Public/PlanarCut.h
--- a/Engine/Plugins/Experimental/PlanarCutPlugin/Source/PlanarCut/Public/PlanarCut.h
+++ b/Engine/Plugins/Experimental/PlanarCutPlugin/Source/PlanarCut/Public/PlanarCut.h
@@ -50,7 +50,7 @@ struct PLANARCUT_API FPlanarCells
 	}
 	FPlanarCells(const FPlane& Plane);
 	FPlanarCells(const TArrayView<const FVector> Sites, FVoronoiDiagram &Voronoi);
-	FPlanarCells(const TArrayView<const FBox> Boxes);
+	FPlanarCells(const TArrayView<const FBox> Boxes, bool bResolveAdjacencies = false);
 	FPlanarCells(const FBox &Region, const FIntVector& CubesPerAxis);
 	FPlanarCells(const FBox &Region, const TArrayView<const FColor> Image, int32 Width, int32 Height);
 
@@ -126,11 +126,12 @@ struct PLANARCUT_API FPlanarCells
 		PlaneBoundaries.Emplace();
 	}
 
-	inline void AddPlane(const FPlane &P, int32 CellIdxBehind, int32 CellIdxInFront, const TArray<int32>& PlaneBoundary)
+	inline int32 AddPlane(const FPlane &P, int32 CellIdxBehind, int32 CellIdxInFront, const TArray<int32>& PlaneBoundary)
 	{
-		Planes.Add(P);
+		int32 PlaneIdx = Planes.Add(P);
 		PlaneCells.Emplace(CellIdxBehind, CellIdxInFront);
 		PlaneBoundaries.Add(PlaneBoundary);
+		return PlaneIdx;
 	}
 
 	void SetNoise(FNoiseSettings Noise = FNoiseSettings())
