compute fractures in a consistent local space, closer to the origin, for more consistent noise + to work better at LWC

diff --git a/Engine/Plugins/Experimental/ChaosEditor/Source/FractureEditor/Private/FractureToolBrick.cpp b/Engine/Plugins/Experimental/ChaosEditor/Source/FractureEditor/Private/FractureToolBrick.cpp
--- a/Engine/Plugins/Experimental/ChaosEditor/Source/FractureEditor/Private/FractureToolBrick.cpp
+++ b/Engine/Plugins/Experimental/ChaosEditor/Source/FractureEditor/Private/FractureToolBrick.cpp
@@ -388,6 +388,7 @@ public:
 
 	TArray<int> Selection;
 	FPlanarCells Cells;
+	FVector CellsOrigin;
 	float PointSpacing;
 	float Grout = 0;
 	int Seed;
@@ -401,7 +402,7 @@ public:
 			return;
 		}
 
-		ResultGeometryIndex = CutMultipleWithPlanarCells(Cells, *CollectionCopy, Selection, Grout, PointSpacing, Seed, Transform, true, true, Progress);
+		ResultGeometryIndex = CutMultipleWithPlanarCells(Cells, *CollectionCopy, Selection, Grout, PointSpacing, Seed, Transform, true, true, Progress, CellsOrigin);
 
 		SetResult(MoveTemp(CollectionCopy));
 	}
@@ -433,9 +434,13 @@ int32 UFractureToolBrick::ExecuteFracture(const FFractureToolContext& FractureCo
 		const FVector HalfBrick(BrickHalfDimensions - HalfGrout);
 		const FBox BrickBox(-HalfBrick, HalfBrick);
 
+		FTransform ContextTransform = FractureContext.GetTransform();
+		FVector Origin = ContextTransform.GetTranslation();
+
 		for (const FTransform& Trans : BrickTransforms)
 		{
-			BricksToCut.Add(BrickBox.TransformBy(Trans));
+			FTransform ToApply = Trans * FTransform(-Origin);
+			BricksToCut.Add(BrickBox.TransformBy(ToApply));
 		}
 
 		TUniquePtr<FCellsFractureOp> BrickOp = MakeUnique<FCellsFractureOp>(*(FractureContext.GetGeometryCollection()));
@@ -452,6 +457,7 @@ int32 UFractureToolBrick::ExecuteFracture(const FFractureToolContext& FractureCo
 		}
 		BrickOp->Seed = FractureContext.GetSeed();
 		BrickOp->Transform = FractureContext.GetTransform();
+		BrickOp->CellsOrigin = Origin;
 
 		int Result = RunCancellableGeometryCollectionOp<FCellsFractureOp>(*(FractureContext.GetGeometryCollection()),
 			MoveTemp(BrickOp), LOCTEXT("ComputingBrickFractureMessage", "Computing Brick Fracture"));
diff --git a/Engine/Plugins/Experimental/ChaosEditor/Source/FractureEditor/Private/FractureToolCutter.cpp b/Engine/Plugins/Experimental/ChaosEditor/Source/FractureEditor/Private/FractureToolCutter.cpp
--- a/Engine/Plugins/Experimental/ChaosEditor/Source/FractureEditor/Private/FractureToolCutter.cpp
+++ b/Engine/Plugins/Experimental/ChaosEditor/Source/FractureEditor/Private/FractureToolCutter.cpp
@@ -380,6 +380,13 @@ public:
 		FProgressCancel::FProgressScope VoronoiDiagramProgress =
 			FProgressCancel::CreateScopeTo(Progress, .01, LOCTEXT("ComputingVoronoiDiagramMessage", "Computing Voronoi Diagram"));
 
+		FVector Origin = Transform.GetTranslation();
+		for (FVector& Site : Sites)
+		{
+			Site -= Origin;
+		}
+		Bounds.Min -= Origin;
+		Bounds.Max -= Origin;
 		FVoronoiDiagram Voronoi(Sites, Bounds, .1f);
 
 		FPlanarCells VoronoiPlanarCells = FPlanarCells(Sites, Voronoi);
@@ -396,7 +403,7 @@ public:
 		FProgressCancel::FProgressScope FractureMeshProgress =
 			FProgressCancel::CreateScopeTo(Progress, 1, LOCTEXT("FractureMeshMessage", "Fracturing Mesh"));
 
-		ResultGeometryIndex = CutMultipleWithPlanarCells(VoronoiPlanarCells, *CollectionCopy, Selection, Grout, PointSpacing, Seed, Transform, true, true, Progress);
+		ResultGeometryIndex = CutMultipleWithPlanarCells(VoronoiPlanarCells, *CollectionCopy, Selection, Grout, PointSpacing, Seed, Transform, true, true, Progress, Origin);
 		
 		SetResult(MoveTemp(CollectionCopy));
 	}
diff --git a/Engine/Plugins/Experimental/PlanarCutPlugin/Source/PlanarCut/Private/PlanarCut.cpp b/Engine/Plugins/Experimental/PlanarCutPlugin/Source/PlanarCut/Private/PlanarCut.cpp
--- a/Engine/Plugins/Experimental/PlanarCutPlugin/Source/PlanarCut/Private/PlanarCut.cpp
+++ b/Engine/Plugins/Experimental/PlanarCutPlugin/Source/PlanarCut/Private/PlanarCut.cpp
@@ -51,6 +51,16 @@ using namespace UE::PlanarCut;
 
 #define LOCTEXT_NAMESPACE "PlanarCut"
 
+namespace PlanarCut_Locals
+{
+	FVector SeparateTranslation(const FTransform& Transform, FTransform& OutCenteredTransform)
+	{
+		FVector Translation = Transform.GetTranslation();
+		OutCenteredTransform = FTransform(Transform.GetRotation(), FVector::ZeroVector, Transform.GetScale3D());
+		return Translation;
+	}
+}
+
 // logic from FMeshUtility::GenerateGeometryCollectionFromBlastChunk, sets material IDs based on construction pattern that external materials have even IDs and are matched to internal materials at InternalID = ExternalID+1
 int32 FInternalSurfaceMaterials::GetDefaultMaterialIDForGeometry(const FGeometryCollection& Collection, int32 GeometryIdx) const
 {
@@ -720,12 +730,12 @@ int32 CutWithPlanarCells(
 	const TOptional<FTransform>& TransformCollection,
 	bool bIncludeOutsideCellInOutput,
 	bool bSetDefaultInternalMaterialsFromCollection,
-	FProgressCancel* Progress
-
+	FProgressCancel* Progress,
+	FVector CellsOrigin
 )
 {
 	TArray<int32> TransformIndices { TransformIdx };
-	return CutMultipleWithPlanarCells(Cells, Source, TransformIndices, Grout, CollisionSampleSpacing, RandomSeed, TransformCollection, bIncludeOutsideCellInOutput, bSetDefaultInternalMaterialsFromCollection, Progress);
+	return CutMultipleWithPlanarCells(Cells, Source, TransformIndices, Grout, CollisionSampleSpacing, RandomSeed, TransformCollection, bIncludeOutsideCellInOutput, bSetDefaultInternalMaterialsFromCollection, Progress, CellsOrigin);
 }
 
 
@@ -757,9 +767,20 @@ int32 CutMultipleWithMultiplePlanes(
 		InternalSurfaceMaterials.SetUVScaleFromCollection(CollectionMesh);
 	}
 
+	// Compute cuts in a local space where the collection has been translated to the origin, for more consistent noise + better LWC accuracy
 	FTransform CollectionToWorld = TransformCollection.Get(FTransform::Identity);
+	FTransform CollectionToWorldCentered;
+	FVector Origin = PlanarCut_Locals::SeparateTranslation(CollectionToWorld, CollectionToWorldCentered);
 
-	FDynamicMeshCollection MeshCollection(&Collection, TransformIndices, CollectionToWorld);
+	// Move planes to the same local space
+	TArray<FPlane> CenteredPlanes;
+	CenteredPlanes.Reserve(Planes.Num());
+	for (const FPlane& Plane : Planes)
+	{
+		CenteredPlanes.Add(Plane.TranslateBy(-Origin));
+	}
+
+	FDynamicMeshCollection MeshCollection(&Collection, TransformIndices, CollectionToWorldCentered);
 
 	PrepareScope.Done();
 	if (Progress && Progress->Cancelled())
@@ -769,7 +790,7 @@ int32 CutMultipleWithMultiplePlanes(
 	FProgressCancel::FProgressScope CutScope = FProgressCancel::CreateScopeTo(Progress, .99, LOCTEXT("CutWithMultiplePlanesBody", "Cutting with planes"));
 
 	int32 NewGeomStartIdx = -1;
-	NewGeomStartIdx = MeshCollection.CutWithMultiplePlanes(Planes, Grout, CollisionSampleSpacing, RandomSeed, &Collection, InternalSurfaceMaterials, bSetDefaultInternalMaterialsFromCollection, Progress);
+	NewGeomStartIdx = MeshCollection.CutWithMultiplePlanes(CenteredPlanes, Grout, CollisionSampleSpacing, RandomSeed, &Collection, InternalSurfaceMaterials, bSetDefaultInternalMaterialsFromCollection, Progress);
 
 	CutScope.Done();
 	if (Progress && Progress->Cancelled())
@@ -796,7 +817,8 @@ int32 CutMultipleWithPlanarCells(
 	const TOptional<FTransform>& TransformCollection,
 	bool bIncludeOutsideCellInOutput,
 	bool bSetDefaultInternalMaterialsFromCollection,
-	FProgressCancel* Progress
+	FProgressCancel* Progress,
+	FVector CellsOrigin
 )
 {
 	FProgressCancel::FProgressScope CreateMeshCollectionScope = FProgressCancel::CreateScopeTo(Progress, .1);
@@ -806,8 +828,10 @@ int32 CutMultipleWithPlanarCells(
 	}
 
 	FTransform CollectionToWorld = TransformCollection.Get(FTransform::Identity);
+	// Put Collection in the same local space as the Cells
+	FTransform CollectionToWorldCentered = CollectionToWorld * FTransform(-CellsOrigin);
 
-	FDynamicMeshCollection MeshCollection(&Source, TransformIndices, CollectionToWorld);
+	FDynamicMeshCollection MeshCollection(&Source, TransformIndices, CollectionToWorldCentered);
 	CreateMeshCollectionScope.Done();
 
 	if (Progress && Progress->Cancelled())
@@ -1006,15 +1030,20 @@ int32 CutWithMesh(
 	ensureMsgf(!InternalSurfaceMaterials.NoiseSettings.IsSet(), TEXT("Noise settings not yet supported for mesh-based fracture"));
 
 	FTransform CollectionToWorld = TransformCollection.Get(FTransform::Identity);
+	FTransform CollectionToWorldCentered;
+	FVector Origin = PlanarCut_Locals::SeparateTranslation(CollectionToWorld, CollectionToWorldCentered);
+
+	FTransform CuttingMeshTransformCentered = CuttingMeshTransform;
+	CuttingMeshTransformCentered.SetTranslation(CuttingMeshTransform.GetTranslation() - Origin);
 
 	if (Progress && Progress->Cancelled())
 	{
 		return -1;
 	}
 
-	FDynamicMeshCollection MeshCollection(&Collection, TransformIndices, CollectionToWorld);
+	FDynamicMeshCollection MeshCollection(&Collection, TransformIndices, CollectionToWorldCentered);
 	int32 NumUVLayers = Collection.NumUVLayers();
-	FCellMeshes CellMeshes(NumUVLayers, DynamicCuttingMesh, InternalSurfaceMaterials, CuttingMeshTransform);
+	FCellMeshes CellMeshes(NumUVLayers, DynamicCuttingMesh, InternalSurfaceMaterials, CuttingMeshTransformCentered);
 
 	TArray<TPair<int32, int32>> CellConnectivity;
 	CellConnectivity.Add(TPair<int32, int32>(0, -1)); // there's only one 'inside' cell (0), so all cut surfaces are connecting the 'inside' cell (0) to the 'outside' cell (-1)
diff --git a/Engine/Plugins/Experimental/PlanarCutPlugin/Source/PlanarCut/Public/PlanarCut.h b/Engine/Plugins/Experimental/PlanarCutPlugin/Source/PlanarCut/Public/PlanarCut.h
--- a/Engine/Plugins/Experimental/PlanarCutPlugin/Source/PlanarCut/Public/PlanarCut.h
+++ b/Engine/Plugins/Experimental/PlanarCutPlugin/Source/PlanarCut/Public/PlanarCut.h
@@ -153,6 +153,7 @@ struct PLANARCUT_API FPlanarCells
  * @param TransformCollection		Optional transform of the whole geometry collection; if unset, defaults to Identity
  * @param bIncludeOutsideCellInOutput	If true, geometry that was not inside any of the cells (e.g. was outside of the bounds of all cutting geometry) will still be included in the output; if false, it will be discarded.
  * @param Progress						Optionally tracks progress and supports early-cancel
+ * @param CellsOrigin					Optionally provide a local origin of the cutting Cells
  * @return	index of first new geometry in the Output GeometryCollection, or -1 if no geometry was added
  */
 int32 PLANARCUT_API CutWithPlanarCells(
@@ -165,7 +166,8 @@ int32 PLANARCUT_API CutWithPlanarCells(
 	const TOptional<FTransform>& TransformCollection = TOptional<FTransform>(),
 	bool bIncludeOutsideCellInOutput = true,
 	bool bSetDefaultInternalMaterialsFromCollection = true,
-	FProgressCancel* Progress = nullptr
+	FProgressCancel* Progress = nullptr,
+	FVector CellsOrigin = FVector::ZeroVector
 );
 
 /**
@@ -180,6 +182,7 @@ int32 PLANARCUT_API CutWithPlanarCells(
  * @param TransformCollection		Optional transform of the whole geometry collection; if unset, defaults to Identity
  * @param bIncludeOutsideCellInOutput	If true, geometry that was not inside any of the cells (e.g. was outside of the bounds of all cutting geometry) will still be included in the output; if false, it will be discarded.
  * @param Progress						Optionally tracks progress and supports early-cancel
+ * @param CellsOrigin					Optionally provide a local origin of the cutting Cells
  * @return	index of first new geometry in the Output GeometryCollection, or -1 if no geometry was added
  */
 int32 PLANARCUT_API CutMultipleWithPlanarCells(
@@ -192,7 +195,8 @@ int32 PLANARCUT_API CutMultipleWithPlanarCells(
 	const TOptional<FTransform>& TransformCollection = TOptional<FTransform>(),
 	bool bIncludeOutsideCellInOutput = true,
 	bool bSetDefaultInternalMaterialsFromCollection = true,
-	FProgressCancel* Progress = nullptr
+	FProgressCancel* Progress = nullptr,
+	FVector CellsOrigin = FVector::ZeroVector
 );
 
 /**
