Add more fracture AutoUV box projection options and default the box project UV origin to (.5, .5)

diff --git a/Engine/Plugins/Experimental/ChaosEditor/Source/FractureEditor/Private/FractureToolUV.cpp b/Engine/Plugins/Experimental/ChaosEditor/Source/FractureEditor/Private/FractureToolUV.cpp
--- a/Engine/Plugins/Experimental/ChaosEditor/Source/FractureEditor/Private/FractureToolUV.cpp
+++ b/Engine/Plugins/Experimental/ChaosEditor/Source/FractureEditor/Private/FractureToolUV.cpp
@@ -220,7 +220,9 @@ void UFractureToolAutoUV::BoxProjectUVs()
 			*GeometryCollectionComponent->GetRestCollection()->GetGeometryCollection(),
 			(FVector3d)AutoUVSettings->ProjectionScale,
 			UseMaterialIDs,
-			AutoUVSettings->TargetMaterialIDs == ETargetMaterialIDs::OddIDs ? EmptyMaterialIDs : AutoUVSettings->MaterialIDs);
+			AutoUVSettings->TargetMaterialIDs == ETargetMaterialIDs::OddIDs ? EmptyMaterialIDs : AutoUVSettings->MaterialIDs,
+			(FVector2f)AutoUVSettings->ProjectionUVOffset,
+			AutoUVSettings->bAutoFitToBounds, AutoUVSettings->bCenterAtPivot, AutoUVSettings->bUniformProjectionScale);
 
 		GeometryCollectionComponent->MarkRenderDynamicDataDirty();
 		GeometryCollectionComponent->MarkRenderStateDirty();
diff --git a/Engine/Plugins/Experimental/ChaosEditor/Source/FractureEditor/Private/FractureToolUV.h b/Engine/Plugins/Experimental/ChaosEditor/Source/FractureEditor/Private/FractureToolUV.h
--- a/Engine/Plugins/Experimental/ChaosEditor/Source/FractureEditor/Private/FractureToolUV.h
+++ b/Engine/Plugins/Experimental/ChaosEditor/Source/FractureEditor/Private/FractureToolUV.h
@@ -109,9 +109,25 @@ public:
 	void DisableBoneColors();
 
 	/** The scale factor to use for UV box projection */
-	UPROPERTY(EditAnywhere, Category = Unwrap)
+	UPROPERTY(EditAnywhere, Category = Unwrap, meta = (EditCondition = "!bAutoFitToBounds"))
 	FVector ProjectionScale = FVector(100, 100, 100);
 
+	/** Set the scale factor for UV box projection based on the bounding box of the geometry */
+	UPROPERTY(EditAnywhere, Category = Unwrap, AdvancedDisplay)
+	bool bAutoFitToBounds = false;
+
+	/** Ensure the projection scale is the same on each axis, either by repeating the X Axis scale, or using the max box dimension if Auto Fit to Bounds is selected. */
+	UPROPERTY(EditAnywhere, Category = Unwrap, AdvancedDisplay, meta = (DisplayName = "Uniform Scale"))
+	bool bUniformProjectionScale = false;
+
+	/** UV Offset to apply after UV box projection */
+	UPROPERTY(EditAnywhere, Category = Unwrap, AdvancedDisplay, meta = (DisplayName = "UV Offset"))
+	FVector2D ProjectionUVOffset = FVector2D(.5, .5);
+
+	/** Optionally center the UV projection around the object pivot */
+	UPROPERTY(EditAnywhere, Category = Unwrap, AdvancedDisplay)
+	bool bCenterAtPivot = false;
+
 	/** Immediately perform a box projection to update the internal UVs of the selected Geometry Collection */
 	UFUNCTION(CallInEditor, Category = Unwrap, meta = (DisplayName = "Box Project UVs"))
 	void BoxProjectUVs();
diff --git a/Engine/Plugins/Experimental/PlanarCutPlugin/Source/PlanarCut/Private/FractureAutoUV.cpp b/Engine/Plugins/Experimental/PlanarCutPlugin/Source/PlanarCut/Private/FractureAutoUV.cpp
--- a/Engine/Plugins/Experimental/PlanarCutPlugin/Source/PlanarCut/Private/FractureAutoUV.cpp
+++ b/Engine/Plugins/Experimental/PlanarCutPlugin/Source/PlanarCut/Private/FractureAutoUV.cpp
@@ -374,7 +374,11 @@ bool BoxProjectUVs(
 	FGeometryCollection& Collection,
 	const FVector3d& BoxDimensions,
 	EUseMaterials MaterialsPattern,
-	TArrayView<int32> WhichMaterials
+	TArrayView<int32> WhichMaterials,
+	FVector2f OffsetUVs,
+	bool bOverrideBoxDimensionsWithBounds,
+	bool bCenterBoxAtPivot,
+	bool bUniformProjectionScale
 )
 {
 	TArray<int32> TransformIndices;
@@ -423,7 +427,27 @@ bool BoxProjectUVs(
 		{
 			FDynamicMeshUVEditor UVEd(&Mesh, TargetUVLayer, false);
 			FFrame3d BoxFrame; // defaults to origin / no rotation
-			UVEd.SetTriangleUVsFromBoxProjection(TargetTris, [](const FVector3d& Pos) { return Pos; }, BoxFrame, BoxDimensions, 1);
+			FVector3d UseBoxDimensions = bUniformProjectionScale ? FVector3d(BoxDimensions.X) : BoxDimensions;
+			if (!bCenterBoxAtPivot || bOverrideBoxDimensionsWithBounds)
+			{
+				FAxisAlignedBox3d Bounds = Mesh.GetBounds(true);
+				if (!bCenterBoxAtPivot)
+				{
+					BoxFrame.Origin = Bounds.Center();
+				}
+				if (bOverrideBoxDimensionsWithBounds)
+				{
+					UseBoxDimensions = bUniformProjectionScale ? FVector3d(Bounds.MaxDim()) : Bounds.Max - Bounds.Min;
+				}
+			}
+			UVEd.SetTriangleUVsFromBoxProjection(TargetTris, [](const FVector3d& Pos) { return Pos; }, BoxFrame, UseBoxDimensions, 1);
+
+			// apply offset to UV space
+			FDynamicMeshUVOverlay* Overlap = UVEd.GetOverlay();
+			for (int ElemID : Overlap->ElementIndicesItr())
+			{
+				Overlap->SetElement(ElemID, OffsetUVs + Overlap->GetElement(ElemID));
+			}
 		}
 
 		Mesh.CompactInPlace();
diff --git a/Engine/Plugins/Experimental/PlanarCutPlugin/Source/PlanarCut/Public/FractureAutoUV.h b/Engine/Plugins/Experimental/PlanarCutPlugin/Source/PlanarCut/Public/FractureAutoUV.h
--- a/Engine/Plugins/Experimental/PlanarCutPlugin/Source/PlanarCut/Public/FractureAutoUV.h
+++ b/Engine/Plugins/Experimental/PlanarCutPlugin/Source/PlanarCut/Public/FractureAutoUV.h
@@ -31,7 +31,11 @@ bool PLANARCUT_API BoxProjectUVs(
 	FGeometryCollection& Collection,
 	const FVector3d& BoxDimensions,
 	EUseMaterials MaterialsPattern = EUseMaterials::OddMaterials,
-	TArrayView<int32> WhichMaterials = TArrayView<int32>()
+	TArrayView<int32> WhichMaterials = TArrayView<int32>(),
+	FVector2f OffsetUVs = FVector2f(.5f,.5f),
+	bool bOverrideBoxDimensionsWithBounds = false,
+	bool bCenterBoxAtPivot = false,
+	bool bUniformProjectionScale = false
 );
 
 /**
