PlanarCut - fixing truncation warnings

diff --git a/Engine/Plugins/Experimental/PlanarCutPlugin/Source/PlanarCut/Private/FractureAutoUV.cpp b/Engine/Plugins/Experimental/PlanarCutPlugin/Source/PlanarCut/Private/FractureAutoUV.cpp
--- a/Engine/Plugins/Experimental/PlanarCutPlugin/Source/PlanarCut/Private/FractureAutoUV.cpp
+++ b/Engine/Plugins/Experimental/PlanarCutPlugin/Source/PlanarCut/Private/FractureAutoUV.cpp
@@ -604,7 +604,7 @@ bool UVLayout(
 	// To approximate a larger gutter, we tell it to consider a smaller output resolution --
 	//  hoping that e.g. the UVs for a 1 pixel gutter at 256x256 are ~ a 2 pixel gutter at 512x512
 	// TODO: If we make StandardPack support the GutterSize parameter, we can use that instead.
-	Packer.TextureResolution = UVRes / FMathf::Max(1, GutterSize);
+	Packer.TextureResolution = static_cast<int32>( UVRes / FMathf::Max(1.f, GutterSize) );
 	return Packer.StandardPack(&UVMesh, UVIslands);
 }
 
@@ -666,11 +666,11 @@ bool TextureInternalSurfaces(
 	}
 
 	const float AmbientWorkPer = 1, CurvatureWorkPer = 10;
-	const float NumMeshes = TransformIndices.Num();
+	const int32 NumMeshes = TransformIndices.Num();
 	const float AmountOfWork = 
 		float(AmbientIdx > -1) * AmbientWorkPer * NumMeshes +
 		float(CurvatureIdx > -1) * CurvatureWorkPer * NumMeshes;
-	const float WorkIncr = AmountOfWork > 0 ? .9 / AmountOfWork : 0;
+	const float WorkIncr = AmountOfWork > 0 ? 0.9f / AmountOfWork : 0.f;
 
 	FDynamicMeshCollection CollectionMeshes(&Collection, TransformIndices, FTransform::Identity, false);
 	if (bNeedsDynamicMeshes)
@@ -877,7 +877,7 @@ bool TextureInternalSurfaces(
 					{
 						double DistanceSq;
 						OutsideSpatial.FindNearestTriangle(InsidePoint, DistanceSq, AttributeSettings.ToExternal_MaxDistance);
-						float PercentDistance = FMathf::Min(1.0f, FMathf::Sqrt(DistanceSq) / (float)AttributeSettings.ToExternal_MaxDistance);
+						float PercentDistance = FMathf::Min(1.0f, static_cast<float>(FMathd::Sqrt(DistanceSq) / AttributeSettings.ToExternal_MaxDistance));
 						checkSlow(FMath::IsFinite(PercentDistance));
 						OutColor[DistanceToExternalIdx] = PercentDistance;
 					}
@@ -887,7 +887,7 @@ bool TextureInternalSurfaces(
 						{
 							double Min = OutsideSpatial.GetBoundingBox().Min[Dim], Max = OutsideSpatial.GetBoundingBox().Max[Dim];
 							checkSlow(Min != Max);
-							float PercentAlong = (float)(InsidePoint[Dim] - Min) / (Max - Min);
+							float PercentAlong = static_cast<float>( (InsidePoint[Dim] - Min) / (Max - Min) );
 							OutColor[TargetIdx] = PercentAlong;
 						}
 					};
diff --git a/Engine/Plugins/Experimental/PlanarCutPlugin/Source/PlanarCut/Private/GeometryMeshConversion.cpp b/Engine/Plugins/Experimental/PlanarCutPlugin/Source/PlanarCut/Private/GeometryMeshConversion.cpp
--- a/Engine/Plugins/Experimental/PlanarCutPlugin/Source/PlanarCut/Private/GeometryMeshConversion.cpp
+++ b/Engine/Plugins/Experimental/PlanarCutPlugin/Source/PlanarCut/Private/GeometryMeshConversion.cpp
@@ -201,8 +201,8 @@ namespace AugmentedDynamicMesh
 	bool GetVisibility(const FDynamicMesh3& Mesh, int TID)
 	{
 		checkSlow(IsAugmented(Mesh));
-		TDynamicMeshScalarTriangleAttribute<bool>* Visible =
-			static_cast<TDynamicMeshScalarTriangleAttribute<bool>*>(Mesh.Attributes()->GetAttachedAttribute(VisibleAttribName));
+		const TDynamicMeshScalarTriangleAttribute<bool>* Visible =
+			static_cast<const TDynamicMeshScalarTriangleAttribute<bool>*>(Mesh.Attributes()->GetAttachedAttribute(VisibleAttribName));
 		return Visible->GetValue(TID);
 	}
 
@@ -244,8 +244,8 @@ namespace AugmentedDynamicMesh
 			UV = FVector2f::Zero();
 			return;
 		}
-		TDynamicMeshVertexAttribute<float, 2>* UVs =
-			static_cast<TDynamicMeshVertexAttribute<float, 2>*>(Mesh.Attributes()->GetAttachedAttribute(UVChannelNames[UVLayer]));
+		const TDynamicMeshVertexAttribute<float, 2>* UVs =
+			static_cast<const TDynamicMeshVertexAttribute<float, 2>*>(Mesh.Attributes()->GetAttachedAttribute(UVChannelNames[UVLayer]));
 		if (ensure(UVs))
 		{
 			UVs->GetValue(VID, UV);
@@ -266,10 +266,10 @@ namespace AugmentedDynamicMesh
 	void GetTangent(const FDynamicMesh3& Mesh, int VID, FVector3f& U, FVector3f& V)
 	{
 		checkSlow(IsAugmented(Mesh));
-		TDynamicMeshVertexAttribute<float, 3>* Us =
-			static_cast<TDynamicMeshVertexAttribute<float, 3>*>(Mesh.Attributes()->GetAttachedAttribute(TangentUAttribName));
-		TDynamicMeshVertexAttribute<float, 3>* Vs =
-			static_cast<TDynamicMeshVertexAttribute<float, 3>*>(Mesh.Attributes()->GetAttachedAttribute(TangentVAttribName));
+		const TDynamicMeshVertexAttribute<float, 3>* Us =
+			static_cast<const TDynamicMeshVertexAttribute<float, 3>*>(Mesh.Attributes()->GetAttachedAttribute(TangentUAttribName));
+		const TDynamicMeshVertexAttribute<float, 3>* Vs =
+			static_cast<const TDynamicMeshVertexAttribute<float, 3>*>(Mesh.Attributes()->GetAttachedAttribute(TangentVAttribName));
 		FVector3f Normal = Mesh.GetVertexNormal(VID);
 		Us->GetValue(VID, U);
 		Vs->GetValue(VID, V);
@@ -663,7 +663,7 @@ namespace AugmentedDynamicMesh
 
 					// Pick number of samples based on the longest edge
 					double LongEdgeLen = FMathd::Sqrt(MaxEdgeLenSq);
-					int Divisions = FMathd::Floor(LongEdgeLen / Spacing);
+					int Divisions = FMath::FloorToInt32(LongEdgeLen / Spacing);
 					double Factor = 1.0 / double(Divisions + 1);
 					int SecondEdgeIdx = (MaxEdgeIdx + 1) % 3;
 					int ThirdEdgeIdx = (MaxEdgeIdx + 2) % 3;
@@ -685,7 +685,7 @@ namespace AugmentedDynamicMesh
 
 						// Choose number of samples between the two edge points based on their distance
 						double AcrossDist = Distance( Triangle.BarycentricPoint(E1Bary), Triangle.BarycentricPoint(E2Bary) );
-						int DivisionsAcross = FMathd::Ceil(AcrossDist / Spacing);
+						int DivisionsAcross = FMath::CeilToInt32(AcrossDist / Spacing);
 						double FactorAcross = 1.0 / double(DivisionsAcross + 1);
 						for (int DivJ = 0; DivJ < DivisionsAcross; DivJ++)
 						{
@@ -1889,7 +1889,7 @@ void FCellMeshes::CreateMeshesForBoundedPlanesWithoutNoise(int NumCells, const F
 			for (int BoundaryVertex : PlaneBoundary)
 			{
 				PlaneVertInfo.Position = FVector3d(Cells.PlaneBoundaryVertices[BoundaryVertex]);
-				FVector2f UV = (FVector2f(PlaneFrame.ToPlaneUV(PlaneVertInfo.Position)) - MinUV) * GlobalUVScale;
+				FVector2f UV = (FVector2f(PlaneFrame.ToPlaneUV(PlaneVertInfo.Position)) - MinUV) * static_cast<float>(GlobalUVScale);
 				int VID = Meshes[MeshIdx]->AppendVertex(PlaneVertInfo);
 				AugmentedDynamicMesh::SetAllUV(*Meshes[MeshIdx], VID, UV, NumUVLayers);
 			}
@@ -2022,7 +2022,7 @@ void FCellMeshes::CreateMeshesForBoundedPlanesWithNoise(int NumCells, const FPla
 			const FVector& V2 = Cells.PlaneBoundaryVertices[PlaneBoundary[V2Idx]];
 			AreaVec += (V1 - V0) ^ (V2 - V1);
 		}
-		TotalArea += AreaVec.Size();
+		TotalArea += static_cast<float>(AreaVec.Size());
 	}
 	double Spacing = GetSafeNoiseSpacing(TotalArea, Cells.InternalSurfaceMaterials.NoiseSettings->PointSpacing);
 
@@ -2312,7 +2312,7 @@ void FCellMeshes::CreateMeshesForBoundedPlanesWithNoise(int NumCells, const FPla
 				FIndex3i Tri = Mesh.GetTriangle(TID);
 				for (int Idx = 0; Idx < 3; Idx++)
 				{
-					FVector2f UV = ((FVector2f)PlaneFrames[PlaneIdx].ToPlaneUV(Mesh.GetVertex(Tri[Idx])) - PlaneMinUVs[PlaneIdx]) * GlobalUVScale;
+					FVector2f UV = ((FVector2f)PlaneFrames[PlaneIdx].ToPlaneUV(Mesh.GetVertex(Tri[Idx])) - PlaneMinUVs[PlaneIdx]) * static_cast<float>(GlobalUVScale);
 					for (int UVLayerIdx = 0; UVLayerIdx < NumUVLayers; UVLayerIdx++)
 					{
 						AugmentedDynamicMesh::SetUV(Mesh, Tri[Idx], UV, UVLayerIdx);
@@ -2356,7 +2356,7 @@ void FCellMeshes::CreateMeshesForSinglePlane(const FPlanarCells& Cells, const FA
 	for (int CornerIdx = 0; CornerIdx < 4; CornerIdx++)
 	{
 		PlaneVertInfo.Position = PlaneFrame.FromPlaneUV(XYRange.GetCorner(CornerIdx));
-		FVector2f UV = FVector2f(XYRange.GetCorner(CornerIdx) - XYRange.Min) * GlobalUVScale;
+		FVector2f UV = FVector2f(XYRange.GetCorner(CornerIdx) - XYRange.Min) * static_cast<float>(GlobalUVScale);
 		int VID = PlaneMesh.AppendVertex(PlaneVertInfo);
 		AugmentedDynamicMesh::SetAllUV(PlaneMesh, VID, UV, NumUVLayers);
 	}
@@ -2365,7 +2365,7 @@ void FCellMeshes::CreateMeshesForSinglePlane(const FPlanarCells& Cells, const FA
 
 	if (bNoise)
 	{
-		double Spacing = GetSafeNoiseSpacing(XYRange.Area(), Cells.InternalSurfaceMaterials.NoiseSettings->PointSpacing);
+		double Spacing = GetSafeNoiseSpacing(static_cast<float>(XYRange.Area()), Cells.InternalSurfaceMaterials.NoiseSettings->PointSpacing);
 		RemeshForNoise(PlaneMesh, EEdgeRefineFlags::SplitsOnly, Spacing);
 		ApplyNoise(PlaneMesh, PlaneFrame.GetAxis(2), Cells.InternalSurfaceMaterials.NoiseSettings.GetValue(), true);
 		FMeshNormals::QuickComputeVertexNormals(PlaneMesh);
@@ -2638,7 +2638,7 @@ int32 FDynamicMeshCollection::CutWithMultiplePlanes(
 		FMeshIndexMappings IndexMaps;
 		for (int32 PlaneIdx = 0; PlaneIdx < Planes.Num(); PlaneIdx++)
 		{
-			FProgressCancel::FProgressScope MakeMeshScope(Progress, .1 * PerPlaneWorkFrac);
+			FProgressCancel::FProgressScope MakeMeshScope(Progress, .1f * PerPlaneWorkFrac);
 			FPlanarCells PlaneCells(Planes[PlaneIdx]);
 			PlaneCells.InternalSurfaceMaterials = InternalSurfaceMaterials;
 			FCellMeshes PlaneGroutMesh(NumUVLayers, RandomStream);
diff --git a/Engine/Plugins/Experimental/PlanarCutPlugin/Source/PlanarCut/Private/PlanarCut.cpp b/Engine/Plugins/Experimental/PlanarCutPlugin/Source/PlanarCut/Private/PlanarCut.cpp
--- a/Engine/Plugins/Experimental/PlanarCutPlugin/Source/PlanarCut/Private/PlanarCut.cpp
+++ b/Engine/Plugins/Experimental/PlanarCutPlugin/Source/PlanarCut/Private/PlanarCut.cpp
@@ -119,7 +119,7 @@ void FInternalSurfaceMaterials::SetUVScaleFromCollection(const GeometryCollectio
 		FaceStart = FaceStartArray[GeometryIdx];
 		FaceEnd = FaceCountArray[GeometryIdx] + FaceStartArray[GeometryIdx];
 	}
-	float UVDistance = 0;
+	double UVDistance = 0;
 	float WorldDistance = 0;
 	for (int32 FaceIdx = FaceStart; FaceIdx < FaceEnd; FaceIdx++)
 	{
@@ -134,7 +134,7 @@ void FInternalSurfaceMaterials::SetUVScaleFromCollection(const GeometryCollectio
 
 	if (WorldDistance > 0)
 	{
-		GlobalUVScale =  UVDistance / WorldDistance;
+		GlobalUVScale =  static_cast<float>(UVDistance) / WorldDistance;
 	}
 	if (GlobalUVScale <= 0)
 	{
@@ -452,20 +452,20 @@ FPlanarCells::FPlanarCells(const FBox& Region, const FIntVector& CubesPerAxis)
 			}
 		}
 	}
-	float Z = Region.Min.Z;
+	float Z = static_cast<float>( Region.Min.Z );
 	int32 ZSliceSize = VertsPerAxis.X * VertsPerAxis.Y;
 	int32 VIdxOffs[8] = { 0, 1, VertsPerAxis.X + 1, VertsPerAxis.X, ZSliceSize, ZSliceSize + 1, ZSliceSize + VertsPerAxis.X + 1, ZSliceSize + VertsPerAxis.X };
-	for (int32 Zi = 0; Zi < CubesPerAxis.Z; Zi++, Z += CellSizes.Z)
+	for (int32 Zi = 0; Zi < CubesPerAxis.Z; Zi++, Z += static_cast<float>(CellSizes.Z))
 	{
-		float Y = Region.Min.Y;
-		float ZN = Z + CellSizes.Z;
-		for (int32 Yi = 0; Yi < CubesPerAxis.Y; Yi++, Y += CellSizes.Y)
+		float Y = static_cast<float>( Region.Min.Y );
+		float ZN = Z + static_cast<float>(CellSizes.Z);
+		for (int32 Yi = 0; Yi < CubesPerAxis.Y; Yi++, Y += static_cast<float>(CellSizes.Y))
 		{
-			float X = Region.Min.X;
-			float YN = Y + CellSizes.Y;
-			for (int32 Xi = 0; Xi < CubesPerAxis.X; Xi++, X += CellSizes.X)
+			float X = static_cast<float>(Region.Min.X);
+			float YN = Y + static_cast<float>(CellSizes.Y);
+			for (int32 Xi = 0; Xi < CubesPerAxis.X; Xi++, X += static_cast<float>(CellSizes.X))
 			{
-				float XN = X + CellSizes.X;
+				float XN = X + static_cast<float>( CellSizes.X );
 				int VIdx = ToIdxUnsafe(VertsPerAxis, Xi, Yi, Zi);
 				int BoxIdx = ToIdxUnsafe(CubesPerAxis, Xi, Yi, Zi);
 
diff --git a/Engine/Plugins/Experimental/PlanarCutPlugin/Source/PlanarCut/Public/PlanarCut.h b/Engine/Plugins/Experimental/PlanarCutPlugin/Source/PlanarCut/Public/PlanarCut.h
--- a/Engine/Plugins/Experimental/PlanarCutPlugin/Source/PlanarCut/Public/PlanarCut.h
+++ b/Engine/Plugins/Experimental/PlanarCutPlugin/Source/PlanarCut/Public/PlanarCut.h
@@ -146,7 +146,7 @@ struct PLANARCUT_API FPlanarCells
 			
 			for (int32 VIdx : Bdry)
 			{
-				float SD = P.PlaneDot(PlaneBoundaryVertices[VIdx]);
+				const double SD = P.PlaneDot(PlaneBoundaryVertices[VIdx]);
 				if (FMath::Abs(SD) > 1e-4)
 				{
 					return false; // vertices should be on plane!
@@ -156,7 +156,7 @@ struct PLANARCUT_API FPlanarCells
 			{
 				return false;
 			}
-			float AngleMeasure = (NormalDir ^ N).SizeSquared();
+			const double AngleMeasure = (NormalDir ^ N).SizeSquared();
 			if (AngleMeasure > 1e-3) // vectors aren't directionally aligned?
 			{
 				return false;
