CIS fix extra line in fracture auto UV from merge

diff --git a/Engine/Plugins/Experimental/PlanarCutPlugin/Source/PlanarCut/Public/FractureAutoUV.h b/Engine/Plugins/Experimental/PlanarCutPlugin/Source/PlanarCut/Public/FractureAutoUV.h
--- a/Engine/Plugins/Experimental/PlanarCutPlugin/Source/PlanarCut/Public/FractureAutoUV.h
+++ b/Engine/Plugins/Experimental/PlanarCutPlugin/Source/PlanarCut/Public/FractureAutoUV.h
@@ -101,7 +101,6 @@ bool PLANARCUT_API UVLayout(
 	TArrayView<int32> WhichMaterials = TArrayView<int32>(),
 	bool bRecreateUVsForDegenerateIslands = true,
 	FProgressCancel* Progress = nullptr
-	TArrayView<int32> WhichMaterials = TArrayView<int32>()
 );
 
 /**
