add missing const to geometrymeshconversion getinternal helper

diff --git a/Engine/Plugins/Experimental/PlanarCutPlugin/Source/PlanarCut/Private/GeometryMeshConversion.cpp b/Engine/Plugins/Experimental/PlanarCutPlugin/Source/PlanarCut/Private/GeometryMeshConversion.cpp
--- a/Engine/Plugins/Experimental/PlanarCutPlugin/Source/PlanarCut/Private/GeometryMeshConversion.cpp
+++ b/Engine/Plugins/Experimental/PlanarCutPlugin/Source/PlanarCut/Private/GeometryMeshConversion.cpp
@@ -210,8 +210,8 @@ namespace AugmentedDynamicMesh
 	bool GetInternal(const FDynamicMesh3& Mesh, int TID)
 	{
 		checkSlow(IsAugmented(Mesh));
-		TDynamicMeshScalarTriangleAttribute<bool>* Internal =
-			static_cast<TDynamicMeshScalarTriangleAttribute<bool>*>(Mesh.Attributes()->GetAttachedAttribute(InternalAttribName));
+		const TDynamicMeshScalarTriangleAttribute<bool>* Internal =
+			static_cast<const TDynamicMeshScalarTriangleAttribute<bool>*>(Mesh.Attributes()->GetAttachedAttribute(InternalAttribName));
 		return Internal->GetValue(TID);
 	}
 
