Extend TinyGeo with a 'TinyCluster' mode, which merges tiny clusters.

Includes a number of useful options to control e.g., whether the clusters considered should be restricted to the current Fracture Level, or whether to consider merges only within the existing parent clusters, or whether to restrict merges to the current proximity graph.

diff --git a/Engine/Plugins/Experimental/ChaosEditor/Source/FractureEditor/Private/FractureSelectionTools.cpp b/Engine/Plugins/Experimental/ChaosEditor/Source/FractureEditor/Private/FractureSelectionTools.cpp
--- a/Engine/Plugins/Experimental/ChaosEditor/Source/FractureEditor/Private/FractureSelectionTools.cpp
+++ b/Engine/Plugins/Experimental/ChaosEditor/Source/FractureEditor/Private/FractureSelectionTools.cpp
@@ -6,7 +6,7 @@
 #include "EditorSupportDelegates.h"
 #include "GeometryCollection/GeometryCollectionObject.h"
 
-void FFractureSelectionTools::ToggleSelectedBones(UGeometryCollectionComponent* GeometryCollectionComponent, TArray<int32>& BoneIndices, bool bClearCurrentSelection, bool bAdd, bool bSnapToLevel)
+void FFractureSelectionTools::ToggleSelectedBones(UGeometryCollectionComponent* GeometryCollectionComponent, const TArray<int32>& BoneIndices, bool bClearCurrentSelection, bool bAdd, bool bSnapToLevel)
 {
 	if (GeometryCollectionComponent)
 	{
diff --git a/Engine/Plugins/Experimental/ChaosEditor/Source/FractureEditor/Private/FractureToolFixTinyGeo.cpp b/Engine/Plugins/Experimental/ChaosEditor/Source/FractureEditor/Private/FractureToolFixTinyGeo.cpp
--- a/Engine/Plugins/Experimental/ChaosEditor/Source/FractureEditor/Private/FractureToolFixTinyGeo.cpp
+++ b/Engine/Plugins/Experimental/ChaosEditor/Source/FractureEditor/Private/FractureToolFixTinyGeo.cpp
@@ -2,12 +2,15 @@
 
 #include "FractureToolFixTinyGeo.h"
 #include "FractureEditorStyle.h"
+#include "FractureSettings.h"
 #include "PlanarCut.h"
 #include "GeometryCollection/GeometryCollection.h"
 #include "GeometryCollection/GeometryCollectionObject.h"
 #include "GeometryCollection/GeometryCollectionAlgo.h"
+#include "GeometryCollection/Facades/CollectionTransformSelectionFacade.h"
 #include "FractureEditorModeToolkit.h"
 #include "FractureToolContext.h"
+#include "Algo/RemoveIf.h"
 
 #include UE_INLINE_GENERATED_CPP_BY_NAME(FractureToolFixTinyGeo)
 
@@ -79,57 +82,121 @@ TArray<UObject*> UFractureToolFixTinyGeo::GetSettingsObjects() const
 	return Settings;
 }
 
-
-void UFractureToolFixTinyGeo::FractureContextChanged()
+bool UFractureToolFixTinyGeo::CollectTargetBones(FGeometryCollection& Collection, const TArray<int32>& Selection, TArray<int32>& OutSmallIndices, TArray<double>& OutVolumes, double& OutMinVolume)
 {
-	UpdateDefaultRandomSeed();
-	TArray<FFractureToolContext> FractureContexts = GetFractureToolContexts();
+	bool bClusterMode = TinyGeoSettings->MergeType == EMergeType::MergeClusters;
+	bool bRestrictToLevel = bClusterMode && TinyGeoSettings->bOnFractureLevel;
+	const UFractureSettings* FractureSettings = GetDefault<UFractureSettings>();
 
-	ClearVisualizations();
+	FindBoneVolumes(
+		Collection,
+		TArrayView<int32>(), /*Empty array => use all transforms*/
+		OutVolumes,
+		VolDimScale,
+		bClusterMode
+	);
 
-	for (const FFractureToolContext& FractureContext : FractureContexts)
-	{
-		FGeometryCollection& Collection = *FractureContext.GetGeometryCollection();
-		int CollectionIdx = VisualizedCollections.Emplace(FractureContext.GetGeometryCollectionComponent());
-
-		TArray<double> Volumes;
-		FindBoneVolumes(
-			Collection,
-			TArrayView<int32>(), /*Empty array => use all transforms*/
-			Volumes,
-			VolDimScale
-		);
+	double TotalVolume = GetTotalVolume(Collection, OutVolumes);
+	OutMinVolume = GetMinVolume(TotalVolume);
 
-		double MinVolume = GetMinVolume(Volumes);
+	TArray<int32> FilteredSelection;
 
-		TArray<int32> SmallIndices;
+	GeometryCollection::Facades::FCollectionTransformSelectionFacade SelectionFacade(Collection);
 
-		if (TinyGeoSettings->UseBoneSelection == EUseBoneSelection::OnlyMergeSelected)
+	if (TinyGeoSettings->UseBoneSelection == EUseBoneSelection::OnlyMergeSelected)
+	{
+		if (!bClusterMode)
 		{
-			SmallIndices = FractureContext.GetSelection();
+			OutSmallIndices = Selection;
+			SelectionFacade.ConvertSelectionToRigidNodes(OutSmallIndices);
 		}
 		else
 		{
-			FindSmallBones(
-				Collection,
-				TArrayView<int32>(),
-				Volumes,
-				MinVolume,
-				SmallIndices
-			);
+			OutSmallIndices = Selection;
+		}
+		return !OutSmallIndices.IsEmpty();
+	}
 
-			if (TinyGeoSettings->UseBoneSelection == EUseBoneSelection::AlsoMergeSelected)
-			{
-				const TArray<int32>& Selection = FractureContext.GetSelection();
-				for (int32 Bone : Selection)
-				{
-					// note: we AddUnique when doing the merge, but for rendering we allow duplicates
-					SmallIndices.Add(Bone);
-				}
-			}
+	FindSmallBones(
+		Collection,
+		TArrayView<int32>(),
+		OutVolumes,
+		OutMinVolume,
+		OutSmallIndices,
+		bClusterMode
+	);
+
+	// Filter bones that aren't at the target level
+	int32 TargetLevel = FractureSettings->FractureLevel;
+	const TManagedArray<int32>* LevelAttrib = Collection.FindAttribute<int32>("Level", FGeometryCollection::TransformGroup);
+	bool bFilterByLevel = bRestrictToLevel && TargetLevel > -1 && LevelAttrib;
+	if (bFilterByLevel)
+	{
+		OutSmallIndices.SetNum(Algo::RemoveIf(OutSmallIndices, [&](int32 BoneIdx)
+		{
+			return (*LevelAttrib)[BoneIdx] != TargetLevel;
+		}));
+	}
+	// Filter bones that aren't clusters
+	if (bClusterMode && (!bFilterByLevel || TinyGeoSettings->bOnlyClusters))
+	{
+		OutSmallIndices.SetNum(Algo::RemoveIf(OutSmallIndices, [&](int32 BoneIdx)
+		{
+			return Collection.SimulationType[BoneIdx] != FGeometryCollection::ESimulationTypes::FST_Clustered;
+		}));
+	}
+
+	if (TinyGeoSettings->UseBoneSelection == EUseBoneSelection::AlsoMergeSelected)
+	{
+		TArray<int32> ProcessedSelection = Selection;
+		if (!bClusterMode)
+		{
+			SelectionFacade.ConvertSelectionToRigidNodes(ProcessedSelection);
 		}
+		for (int32 Bone : ProcessedSelection)
+		{
+			OutSmallIndices.AddUnique(Bone);
+		}
+	}
+
+	return !OutSmallIndices.IsEmpty();
+}
+
+
+void UFractureToolFixTinyGeo::FractureContextChanged()
+{
+	const UFractureSettings* FractureSettings = GetDefault<UFractureSettings>();
+	bool bLevelIsAll = FractureSettings->FractureLevel == -1;
+	if (TinyGeoSettings->bFractureLevelIsAll != bLevelIsAll)
+	{
+		TinyGeoSettings->bFractureLevelIsAll = bLevelIsAll;
+		NotifyOfPropertyChangeByTool(TinyGeoSettings);
+	}
 
-		FTransform OuterTransform = FractureContext.GetTransform();
+	UpdateDefaultRandomSeed();
+	TSet<UGeometryCollectionComponent*> GeomCompSelection;
+	GetSelectedGeometryCollectionComponents(GeomCompSelection);
+
+	ClearVisualizations();
+
+	TArray<int32> Leaves;
+	for (UGeometryCollectionComponent* GeometryCollectionComponent : GeomCompSelection)
+	{
+		// Get a non-const rest collection to support update of the Volume attribute
+		FGeometryCollectionEdit EditCollection = GeometryCollectionComponent->EditRestCollection(GeometryCollection::EEditUpdate::None, true);
+		FGeometryCollection& Collection = *EditCollection.GetRestCollection()->GetGeometryCollection();
+		GeometryCollection::Facades::FCollectionTransformSelectionFacade SelectionFacade(Collection);
+
+		TArray<int32> SelectedBones = GeometryCollectionComponent->GetSelectedBones();
+		
+		int CollectionIdx = VisualizedCollections.Emplace(GeometryCollectionComponent);
+
+		TArray<int32> SmallIndices;
+		TArray<double> Volumes;
+		double MinVolume;
+		CollectTargetBones(Collection, SelectedBones, SmallIndices, Volumes, MinVolume);
+
+		FTransform OuterTransform = GeometryCollectionComponent->GetOwner()->GetActorTransform();
 		for (int32 TransformIdx : SmallIndices) // small transforms
 		{
 			FTransform InnerTransform = GeometryCollectionAlgo::GlobalMatrix(Collection.Transform, Collection.Parent, TransformIdx);
@@ -137,18 +204,23 @@ void UFractureToolFixTinyGeo::FractureContextChanged()
 
 			FTransform CombinedTransform = InnerTransform * OuterTransform;
 
-			int32 GeometryIdx = Collection.TransformToGeometryIndex[TransformIdx];
-			if (GeometryIdx == INDEX_NONE)
-			{
-				ensure(false); // already filtered for small geo, shouldn't have any missing-geo transforms in the list
-				continue;
-			}
-			int32 VStart = Collection.VertexStart[GeometryIdx];
-			int32 VEnd = VStart + Collection.VertexCount[GeometryIdx];
 			FBox Bounds(EForceInit::ForceInit);
-			for (int32 VIdx = VStart; VIdx < VEnd; VIdx++)
+			Leaves.Reset();
+			SelectionFacade.ConvertSelectionToRigidNodes(TransformIdx, Leaves);
+			for (int32 LeafIdx : Leaves)
 			{
-				Bounds += CombinedTransform.TransformPosition((FVector)Collection.Vertex[VIdx]);
+				int32 GeometryIdx = Collection.TransformToGeometryIndex[LeafIdx];
+				if (GeometryIdx == INDEX_NONE)
+				{
+					ensure(false); // already filtered for small geo, shouldn't have any missing-geo transforms in the list
+					continue;
+				}
+				int32 VStart = Collection.VertexStart[GeometryIdx];
+				int32 VEnd = VStart + Collection.VertexCount[GeometryIdx];
+				for (int32 VIdx = VStart; VIdx < VEnd; VIdx++)
+				{
+					Bounds += CombinedTransform.TransformPosition((FVector)Collection.Vertex[VIdx]);
+				}
 			}
 			ToRemoveMappings.AddMapping(CollectionIdx, TransformIdx, ToRemoveBounds.Num());
 			ToRemoveBounds.Add(Bounds);
@@ -156,7 +228,24 @@ void UFractureToolFixTinyGeo::FractureContextChanged()
 	}
 }
 
-double UFractureToolFixTinyGeo::GetMinVolume(TArray<double>& Volumes)
+double UFractureToolFixTinyGeo::GetTotalVolume(const FGeometryCollection& Collection, const TArray<double>& Volumes)
+{
+	double Sum = 0.0;
+	if (Volumes.Num() != Collection.Transform.Num())
+	{
+		return 0.0;
+	}
+	for (int32 BoneIdx = 0, NumBones = Collection.Transform.Num(); BoneIdx < NumBones; ++BoneIdx)
+	{
+		if (Collection.SimulationType[BoneIdx] == FGeometryCollection::ESimulationTypes::FST_Rigid)
+		{
+			Sum += Volumes[BoneIdx];
+		}
+	}
+	return Sum;
+}
+
+double UFractureToolFixTinyGeo::GetMinVolume(double TotalVolume) const
 {
 	double MinVolume = 0;
 	if (TinyGeoSettings->SelectionMethod == EGeometrySelectionMethod::VolumeCubeRoot)
@@ -166,89 +255,92 @@ double UFractureToolFixTinyGeo::GetMinVolume(TArray<double>& Volumes)
 	}
 	else // EGeometrySelectionMethod::RelativeVolume
 	{
-		double VolumeSum = 0;
-		for (double Volume : Volumes)
-		{
-			VolumeSum += Volume;
-		}
-		MinVolume = FMath::Pow(VolumeSum, 1.0 / 3.0) * TinyGeoSettings->RelativeVolume;
+		MinVolume = FMath::Pow(TotalVolume, 1.0 / 3.0) * TinyGeoSettings->RelativeVolume;
 		MinVolume = MinVolume * MinVolume * MinVolume;
 	}
 	return MinVolume;
 }
 
-int32 UFractureToolFixTinyGeo::ExecuteFracture(const FFractureToolContext& FractureContext)
+void UFractureToolFixTinyGeo::Execute(TWeakPtr<FFractureEditorModeToolkit> InToolkit)
 {
-	if (FractureContext.GetGeometryCollection().IsValid())
+	if (!InToolkit.IsValid())
 	{
-		TArray<int32> TransformIndices; // left empty to refer to all indices
-		if (TinyGeoSettings->UseBoneSelection == EUseBoneSelection::OnlyMergeSelected)
-		{
-			// restrict to selection
-			TransformIndices = FractureContext.GetSelection();
-			if (TransformIndices.IsEmpty())
-			{
-				return INDEX_NONE;
-			}
-		}
-
-		TArray<double> Volumes;
-		FGeometryCollection& Collection = *FractureContext.GetGeometryCollection();
-		FindBoneVolumes(
-			Collection,
-			TransformIndices,
-			Volumes,
-			VolDimScale
-		);
+		return;
+	}
+	FFractureEditorModeToolkit* Toolkit = InToolkit.Pin().Get();
 
-		double MinVolume = GetMinVolume(Volumes);
+	TSet<UGeometryCollectionComponent*> GeomCompSelection;
+	GetSelectedGeometryCollectionComponents(GeomCompSelection);
+	for (UGeometryCollectionComponent* GeometryCollectionComponent : GeomCompSelection)
+	{
+		FGeometryCollectionEdit CollectionEdit = GeometryCollectionComponent->EditRestCollection(GeometryCollection::EEditUpdate::RestPhysicsDynamic, false);
+		FGeometryCollection& Collection = *CollectionEdit.GetRestCollection()->GetGeometryCollection();
+		TArray<int32> SelectedBones = GeometryCollectionComponent->GetSelectedBones();
 
 		TArray<int32> SmallIndices;
-		if (TinyGeoSettings->UseBoneSelection == EUseBoneSelection::OnlyMergeSelected)
+		TArray<double> Volumes;
+		double MinVolume;
+		if (!CollectTargetBones(Collection, SelectedBones, SmallIndices, Volumes, MinVolume))
 		{
-			SmallIndices = FractureContext.GetSelection();
+			continue;
 		}
-		else
+
+		// convert ENeighborSelectionMethod to the PlanarCut module version (the only difference is it's not a UENUM)
+		UE::PlanarCut::ENeighborSelectionMethod SelectionMethod = UE::PlanarCut::ENeighborSelectionMethod::LargestNeighbor;
+		if (TinyGeoSettings->NeighborSelection == ENeighborSelectionMethod::NearestCenter)
+		{
+			SelectionMethod = UE::PlanarCut::ENeighborSelectionMethod::NearestCenter;
+		}
+
+		if (TinyGeoSettings->MergeType == EMergeType::MergeGeometry)
 		{
-			FindSmallBones(
+			MergeBones(
 				Collection,
-				TransformIndices,
+				TArrayView<const int32>(), // empty view == consider all bones
 				Volumes,
 				MinVolume,
-				SmallIndices
+				SmallIndices,
+				false, /*bUnionJoinedPieces*/ // Note: Union-ing the pieces is nicer in theory, but can leave cracks and non-manifold garbage
+				SelectionMethod
 			);
-
-			if (TinyGeoSettings->UseBoneSelection == EUseBoneSelection::AlsoMergeSelected)
-			{
-				const TArray<int32>& Selection = FractureContext.GetSelection();
-				for (int32 Bone : Selection)
-				{
-					SmallIndices.AddUnique(Bone);
-				}
-			}
 		}
-
-		// convert ENeighborSelectionMethod to the PlanarCut module version (the only difference is it's not a UENUM)
-		UE::PlanarCut::ENeighborSelectionMethod SelectionMethod = UE::PlanarCut::ENeighborSelectionMethod::LargestNeighbor;
-		if (TinyGeoSettings->NeighborSelection == ENeighborSelectionMethod::NearestCenter)
+		else
 		{
-			SelectionMethod = UE::PlanarCut::ENeighborSelectionMethod::NearestCenter;
+			MergeClusters(
+				Collection,
+				Volumes,
+				MinVolume,
+				SmallIndices,
+				SelectionMethod,
+				TinyGeoSettings->bOnlyToConnected,
+				TinyGeoSettings->bOnlySameParent
+			);
 		}
 
-		MergeBones(
-			Collection,
-			TransformIndices,
-			Volumes,
-			MinVolume,
-			SmallIndices,
-			false, /*bUnionJoinedPieces*/ // Note: Union-ing the pieces is nicer in theory, but can leave cracks and non-manifold garbage
-			SelectionMethod
-		);
+		// Refresh component with cleared selection
+		TArray<int32> EmptySelection;
+		Refresh(GeometryCollectionComponent, Toolkit, EmptySelection, true /*bClearSelection*/, true /*bMustUpdateBoneColors*/);
 	}
 
-	// Don't update selection
+	// Update overall outliner (including stats, etc)
+	Toolkit->SetOutlinerComponents(GeomCompSelection.Array());
+}
+
+int32 UFractureToolFixTinyGeo::ExecuteFracture(const FFractureToolContext& FractureContext)
+{
+	// This tool instead overrides Execute, and should never call ExecuteFracture
+	checkSlow(false);
 	return INDEX_NONE;
 }
 
+void UFractureToolFixTinyGeo::Setup(TWeakPtr<FFractureEditorModeToolkit> InToolkit)
+{
+	Super::Setup(InToolkit);
+	if (InToolkit.IsValid())
+	{
+		InToolkit.Pin()->SetOutlinerColumnMode(EOutlinerColumnMode::Size);
+	}
+}
+
 
 #undef LOCTEXT_NAMESPACE
diff --git a/Engine/Plugins/Experimental/ChaosEditor/Source/FractureEditor/Private/FractureToolFixTinyGeo.h b/Engine/Plugins/Experimental/ChaosEditor/Source/FractureEditor/Private/FractureToolFixTinyGeo.h
--- a/Engine/Plugins/Experimental/ChaosEditor/Source/FractureEditor/Private/FractureToolFixTinyGeo.h
+++ b/Engine/Plugins/Experimental/ChaosEditor/Source/FractureEditor/Private/FractureToolFixTinyGeo.h
@@ -13,7 +13,7 @@ UENUM()
 enum class EGeometrySelectionMethod
 {
 	// Select by cube root of volume
-	VolumeCubeRoot,
+	VolumeCubeRoot UMETA(DisplayName = "Size"),
 	// Select by cube root of volume relative to the overall shape's cube root of volume
 	RelativeVolume
 };
@@ -35,6 +35,13 @@ enum class EUseBoneSelection
 	OnlyMergeSelected
 };
 
+UENUM()
+enum class EMergeType : uint8
+{
+	MergeGeometry,
+	MergeClusters
+};
+
 /** Settings controlling how geometry is selected and merged into neighboring geometry */
 UCLASS(config = EditorPerProjectUserSettings)
 class UFractureTinyGeoSettings : public UFractureToolSettings
@@ -47,9 +54,33 @@ public:
 		: Super(ObjInit)
 	{}
 
+	// Whether to merge small geometry, or small clusters
 	UPROPERTY(EditAnywhere, Category = MergeSettings)
+	EMergeType MergeType = EMergeType::MergeGeometry;
+
+	// Only consider bones at the current Fracture Level
+	UPROPERTY(EditAnywhere, Category = MergeSettings, meta = (EditCondition = "!bFractureLevelIsAll && MergeType == EMergeType::MergeClusters", EditConditionHides))
+	bool bOnFractureLevel = true;
+
+	// Only auto-consider clusters for merging. Note that leaf nodes can still be consider if manually selected.
+	UPROPERTY(EditAnywhere, Category = MergeSettings, meta = (EditCondition = "!bFractureLevelIsAll && MergeType == EMergeType::MergeClusters && bOnFractureLevel", EditConditionHides))
+	bool bOnlyClusters = false;
+
+	// Only merge to neighbors with the same parent in the hierarchy
+	UPROPERTY(EditAnywhere, Category = MergeSettings, meta = (EditCondition = "MergeType == EMergeType::MergeClusters", EditConditionHides))
+	bool bOnlySameParent = true;
+
+	// Helper variable to let the EditConditions above check whether the Fracture Level is set to 'All'
+	UPROPERTY()
+	bool bFractureLevelIsAll = false;
+
+	UPROPERTY(EditAnywhere, Category = MergeSettings, meta = (DisplayName = "Merge To"))
 	ENeighborSelectionMethod NeighborSelection = ENeighborSelectionMethod::LargestNeighbor;
 
+	// Only merge pieces that are connected in the proximity graph. If unchecked, connected pieces will still be favored, but if none are available the closest disconnected piece can be merged.
+	UPROPERTY(EditAnywhere, Category = MergeSettings, meta = (EditCondition = "MergeType == EMergeType::MergeClusters", EditConditionHides))
+	bool bOnlyToConnected = true;
+
 	/** Options for using the current bone selection */
 	UPROPERTY(EditAnywhere, Category = MergeSettings, meta = (DisplayName = "Bone Selection"))
 	EUseBoneSelection UseBoneSelection;
@@ -58,8 +89,8 @@ public:
 	UPROPERTY(EditAnywhere, Category = FilterSettings)
 	EGeometrySelectionMethod SelectionMethod = EGeometrySelectionMethod::RelativeVolume;
 
-	/** If volume is less than this value cubed, geometry should be merged into neighbors -- i.e. a value of 2 merges geometry smaller than a 2x2x2 cube */
-	UPROPERTY(EditAnywhere, Category = FilterSettings, meta = (ClampMin = ".00001", UIMin = ".1", UIMax = "10", EditCondition = "SelectionMethod == EGeometrySelectionMethod::VolumeCubeRoot", EditConditionHides))
+	/** If size (cube root of volume) is less than this value, geometry should be merged into neighbors -- i.e. a value of 2 merges geometry smaller than a 2x2x2 cube */
+	UPROPERTY(EditAnywhere, Category = FilterSettings, meta = (DisplayName = "MinSize", ClampMin = ".00001", UIMin = ".1", UIMax = "10", EditCondition = "SelectionMethod == EGeometrySelectionMethod::VolumeCubeRoot", EditConditionHides))
 	double MinVolumeCubeRoot = 1;
 
 	/** If cube root of volume relative to the overall shape's cube root of volume is less than this, the geometry should be merged into its neighbors.
@@ -85,7 +116,7 @@ public:
 	virtual FText GetTooltipText() const override;
 	virtual FSlateIcon GetToolIcon() const override;
 
-	virtual FText GetApplyText() const override { return FText(NSLOCTEXT("FractureGeoMerge", "ExecuteGeoMerge", "Merge Geometry")); }
+	virtual FText GetApplyText() const override { return FText(NSLOCTEXT("FractureGeoMerge", "ExecuteGeoMerge", "Merge")); }
 
 	void Render(const FSceneView* View, FViewport* Viewport, FPrimitiveDrawInterface* PDI) override;
 
@@ -95,6 +126,8 @@ public:
 
 	virtual void FractureContextChanged() override;
 	virtual int32 ExecuteFracture(const FFractureToolContext& FractureContext) override;
+	virtual void Execute(TWeakPtr<FFractureEditorModeToolkit> InToolkit) override;
+	virtual void Setup(TWeakPtr<FFractureEditorModeToolkit> InToolkit) override;
 
 protected:
 	virtual void ClearVisualizations() override
@@ -112,9 +145,12 @@ private:
 	TArray<FBox> ToRemoveBounds; // Bounds in global space but without exploded vectors applied
 	FVisualizationMappings ToRemoveMappings;
 
-	double GetMinVolume(TArray<double>& Volumes);
+	static double GetTotalVolume(const FGeometryCollection& Collection, const TArray<double>& Volumes);
+	double GetMinVolume(double TotalVolume) const;
 	const double VolDimScale = .01; // compute volumes in meters instead of cm, for saner units at typical scales
 
+	bool CollectTargetBones(FGeometryCollection& Collection, const TArray<int32>& Selection, TArray<int32>& OutSmallIndices, TArray<double>& OutVolumes, double& OutMinVolume);
+
 };
 
 
diff --git a/Engine/Plugins/Experimental/ChaosEditor/Source/FractureEditor/Public/FractureSelectionTools.h b/Engine/Plugins/Experimental/ChaosEditor/Source/FractureEditor/Public/FractureSelectionTools.h
--- a/Engine/Plugins/Experimental/ChaosEditor/Source/FractureEditor/Public/FractureSelectionTools.h
+++ b/Engine/Plugins/Experimental/ChaosEditor/Source/FractureEditor/Public/FractureSelectionTools.h
@@ -15,10 +15,11 @@ public:
 	 *
 	 * @param GeometryCollectionComponent	The component with bones being selected
 	 * @param BoneIndices					The indices to select
-	 * @param bClearCurrentSelection		Whether or not to clear the current selection set or just append the new indices to it				
+	 * @param bClearCurrentSelection		Whether to clear the current selection set or just append the new indices to it				
+	 * @param bAdd							Whether to always add the BoneIndices to the selection, or to toggle them (i.e., de-select if they were previously selected)
 	 * @param bSnapToLevel					Whether to snap the selection to the nearest visible bones (in the 3D view and outliner) by searching the parent bones
 	 */
-	static void ToggleSelectedBones(UGeometryCollectionComponent* GeometryCollectionComponent, TArray<int32>& BoneIndices, bool bClearCurrentSelection, bool bAdd, bool bSnapToLevel = true);
+	static void ToggleSelectedBones(UGeometryCollectionComponent* GeometryCollectionComponent, const TArray<int32>& BoneIndices, bool bClearCurrentSelection, bool bAdd, bool bSnapToLevel = true);
 	static void ClearSelectedBones(UGeometryCollectionComponent* GeometryCollectionComponent);
 
 	static void SelectNeighbors(UGeometryCollectionComponent* GeometryCollectionComponent);
diff --git a/Engine/Plugins/Experimental/ChaosEditor/Source/FractureEditor/Public/FractureTool.cpp b/Engine/Plugins/Experimental/ChaosEditor/Source/FractureEditor/Public/FractureTool.cpp
--- a/Engine/Plugins/Experimental/ChaosEditor/Source/FractureEditor/Public/FractureTool.cpp
+++ b/Engine/Plugins/Experimental/ChaosEditor/Source/FractureEditor/Public/FractureTool.cpp
@@ -242,30 +242,37 @@ TArray<FString> UFractureActionTool::GetSelectedComponentMaterialNames(bool bInc
 	return MaterialNames;
 }
 
-void UFractureActionTool::Refresh(FFractureToolContext& Context, FFractureEditorModeToolkit* Toolkit, bool bClearSelection)
+void UFractureActionTool::Refresh(UGeometryCollectionComponent* GeometryCollectionComponent, FFractureEditorModeToolkit* Toolkit, const TArray<int32>& SetSelection, bool bClearSelection, bool bMustUpdateBoneColors)
 {
-	UGeometryCollectionComponent* GeometryCollectionComponent = Context.GetGeometryCollectionComponent();
-	TSharedPtr<FGeometryCollection, ESPMode::ThreadSafe> GeometryCollectionPtr = Context.GetGeometryCollection();
-	
+	FGeometryCollectionEdit CollectionEdit = GeometryCollectionComponent->EditRestCollection(GeometryCollection::EEditUpdate::None);
+	TSharedPtr<FGeometryCollection, ESPMode::ThreadSafe> GeometryCollectionPtr = CollectionEdit.GetRestCollection()->GetGeometryCollection();
+
 	FGeometryCollectionClusteringUtility::UpdateHierarchyLevelOfChildren(GeometryCollectionPtr.Get(), -1);
 
 	Toolkit->RegenerateOutliner();
 	Toolkit->RegenerateHistogram();
 
-	FScopedColorEdit EditBoneColor(GeometryCollectionComponent, true);
+	if (bMustUpdateBoneColors)
+	{
+		GeometryCollectionComponent->EditBoneSelection(true);
+	}
 
 	if (bClearSelection)
 	{
-		EditBoneColor.ResetBoneSelection();
 		FFractureSelectionTools::ClearSelectedBones(GeometryCollectionComponent);
 	}
 	else
 	{
-		EditBoneColor.SetSelectedBones(Context.GetSelection());
-		FFractureSelectionTools::ToggleSelectedBones(GeometryCollectionComponent, Context.GetSelection(), true, false);
+		FFractureSelectionTools::ToggleSelectedBones(GeometryCollectionComponent, SetSelection, true, true);
 	}
 
-	Context.GetFracturedGeometryCollection()->RebuildRenderData();
+	CollectionEdit.GetRestCollection()->RebuildRenderData();
+}
+
+void UFractureActionTool::Refresh(FFractureToolContext& Context, FFractureEditorModeToolkit* Toolkit, bool bClearSelection)
+{
+	UGeometryCollectionComponent* GeometryCollectionComponent = Context.GetGeometryCollectionComponent();
+	Refresh(GeometryCollectionComponent, Toolkit, Context.GetSelection(), bClearSelection);
 }
 
 void UFractureActionTool::SetOutlinerComponents(TArray<FFractureToolContext>& InContexts, FFractureEditorModeToolkit* Toolkit)
@@ -411,10 +418,6 @@ void UFractureModalTool::Execute(TWeakPtr<FFractureEditorModeToolkit> InToolkit)
 				}
 				FractureContext.Sanitize(false);
 			}
-			
-			Toolkit->RegenerateHistogram();
-
-			FGeometryCollectionClusteringUtility::UpdateHierarchyLevelOfChildren(FractureContext.GetGeometryCollection().Get(), -1);
 
 			Refresh(FractureContext, Toolkit);
 		}
diff --git a/Engine/Plugins/Experimental/ChaosEditor/Source/FractureEditor/Public/FractureTool.h b/Engine/Plugins/Experimental/ChaosEditor/Source/FractureEditor/Public/FractureTool.h
--- a/Engine/Plugins/Experimental/ChaosEditor/Source/FractureEditor/Public/FractureTool.h
+++ b/Engine/Plugins/Experimental/ChaosEditor/Source/FractureEditor/Public/FractureTool.h
@@ -94,6 +94,7 @@ protected:
 	static void AddSingleRootNodeIfRequired(UGeometryCollection* GeometryCollectionObject);
 	static void AddAdditionalAttributesIfRequired(UGeometryCollection* GeometryCollectionObject);
 	static void Refresh(FFractureToolContext& Context, FFractureEditorModeToolkit* Toolkit, bool bClearSelection=false);
+	static void Refresh(UGeometryCollectionComponent* Component, FFractureEditorModeToolkit* Toolkit, const TArray<int32>& SetSelection, bool bClearSelection=false, bool bMustUpdateBoneColors=false);
 	static void SetOutlinerComponents(TArray<FFractureToolContext>& InContexts, FFractureEditorModeToolkit* Toolkit);
 	static void ClearProximity(FGeometryCollection* GeometryCollection);
 
diff --git a/Engine/Plugins/Experimental/PlanarCutPlugin/Source/PlanarCut/Private/PlanarCut.cpp b/Engine/Plugins/Experimental/PlanarCutPlugin/Source/PlanarCut/Private/PlanarCut.cpp
--- a/Engine/Plugins/Experimental/PlanarCutPlugin/Source/PlanarCut/Private/PlanarCut.cpp
+++ b/Engine/Plugins/Experimental/PlanarCutPlugin/Source/PlanarCut/Private/PlanarCut.cpp
@@ -8,6 +8,7 @@
 
 #include "GeometryCollection/GeometryCollectionAlgo.h"
 #include "GeometryCollection/GeometryCollectionClusteringUtility.h"
+#include "GeometryCollection/GeometryCollectionConvexUtility.h"
 
 
 #include "DynamicMeshEditor.h"
@@ -1198,11 +1199,16 @@ void FindBoneVolumes(
 	FGeometryCollection& Collection,
 	const TArrayView<const int32>& TransformIndices,
 	TArray<double>& OutVolumes,
-	double ScalePerDimension
+	double ScalePerDimension,
+	bool bIncludeClusters
 )
 {
 	OutVolumes.Reset();
 
+	// make sure volume attributes are up to date
+	FGeometryCollectionConvexUtility::SetVolumeAttributes(&Collection);
+	const TManagedArray<float>& VolumesAttrib = Collection.GetAttribute<float>("Volume", FTransformCollection::TransformGroup);
+
 	TArray<FTransform> Transforms;
 	TArray<int32> TransformIndicesArray(TransformIndices);
 	if (TransformIndicesArray.Num() == 0)
@@ -1212,74 +1218,44 @@ void FindBoneVolumes(
 			TransformIndicesArray.Add(TransformIdx);
 		}
 	}
-	GeometryCollectionAlgo::GlobalMatrices(Collection.Transform, Collection.Parent, TransformIndicesArray, Transforms);
 
-	auto GetVolume = [](const FGeometryCollection& Collection, int32 GeomIdx, const FTransform& Transform,
-		double DimScaleFactor = 1) -> double
+	double VolumeScale = ScalePerDimension * ScalePerDimension * ScalePerDimension;
+
+	OutVolumes.SetNum(TransformIndicesArray.Num());
+	for (int32 Idx = 0; Idx < TransformIndicesArray.Num(); ++Idx)
 	{
-		int32 VStart = Collection.VertexStart[GeomIdx];
-		int32 VEnd = VStart + Collection.VertexCount[GeomIdx];
-		if (VStart == VEnd)
-		{
-			return 0.0;
-		}
-		FVector3d Center = FVector::ZeroVector;
-		for (int32 VIdx = VStart; VIdx < VEnd; VIdx++)
+		int32 TransformIdx = TransformIndicesArray[Idx];
+		if (bIncludeClusters || Collection.IsRigid(TransformIdx))
 		{
-			FVector Pos = Transform.TransformPosition(FVector(Collection.Vertex[VIdx]));
-			Center += Pos;
+			double UnscaledVolume = (double)VolumesAttrib[Idx];
+			OutVolumes[Idx] = UnscaledVolume * VolumeScale;
 		}
-		Center /= double(VEnd - VStart);
-		int32 FStart = Collection.FaceStart[GeomIdx];
-		int32 FEnd = FStart + Collection.FaceCount[GeomIdx];
-		double VolOut = 0;
-		for (int32 FIdx = FStart; FIdx < FEnd; FIdx++)
+		else
 		{
-			FIntVector Tri = Collection.Indices[FIdx];
-			FVector3d V0 = Transform.TransformPosition(FVector3d(Collection.Vertex[Tri.X]));
-			FVector3d V1 = Transform.TransformPosition(FVector3d(Collection.Vertex[Tri.Y]));
-			FVector3d V2 = Transform.TransformPosition(FVector3d(Collection.Vertex[Tri.Z]));
-
-			// add volume of the tetrahedron formed by the triangles and the reference point
-			FVector3d V1mRef = (V1 - Center) * DimScaleFactor;
-			FVector3d V2mRef = (V2 - Center) * DimScaleFactor;
-			FVector3d N = V2mRef.Cross(V1mRef);
-
-			VolOut += ((V0 - Center) * DimScaleFactor).Dot(N) / 6.0;
+			OutVolumes[Idx] = 0.0;
 		}
-		return VolOut;
-	};
-
-	OutVolumes.SetNum(TransformIndicesArray.Num());
-	ParallelFor(TransformIndicesArray.Num(), [&](int32 Idx)
-		{
-			int32 TransformIdx = TransformIndicesArray[Idx];
-			int32 GeomIdx = Collection.TransformToGeometryIndex[TransformIdx];
-			if (GeomIdx == -1 || !Collection.IsRigid(TransformIdx))
-			{
-				OutVolumes[Idx] = 0.0;
-			}
-			else
-			{
-				OutVolumes[Idx] = GetVolume(Collection, GeomIdx, Transforms[Idx], ScalePerDimension);
-			}
-		});
+	}
 }
 
 
 void FilterBonesByVolume(
-	FGeometryCollection& Collection,
+	const FGeometryCollection& Collection,
 	const TArrayView<const int32>& TransformIndices,
 	const TArrayView<const double>& Volumes,
 	TFunctionRef<bool(double Volume, int32 BoneIdx)> Filter,
-	TArray<int32>& OutSmallBones
+	TArray<int32>& OutSmallBones,
+	bool bIncludeClusters
 )
 {
 	OutSmallBones.Reset();
 
-	auto AddIdx = [&Collection, &Volumes, &Filter, &OutSmallBones](int32 TransformIdx, double Volume)
+	auto AddIdx = [&Collection, &Volumes, &Filter, &OutSmallBones, bIncludeClusters](int32 TransformIdx, double Volume)
 	{
-		if (Collection.IsRigid(TransformIdx) && Collection.TransformToGeometryIndex[TransformIdx] > -1 && Filter(Volume, TransformIdx))
+		constexpr int8 ClusterType = FGeometryCollection::ESimulationTypes::FST_Clustered;
+		constexpr int8 RigidType = FGeometryCollection::ESimulationTypes::FST_Rigid;
+		const int32 SimType = Collection.SimulationType[TransformIdx];
+		const bool bCanConsiderIdx = (bIncludeClusters && SimType == ClusterType) || (SimType == RigidType && Collection.TransformToGeometryIndex[TransformIdx] > -1);
+		if (bCanConsiderIdx && Filter(Volume, TransformIdx))
 		{
 			OutSmallBones.Add(TransformIdx);
 		}
@@ -1314,11 +1290,12 @@ void FilterBonesByVolume(
 
 
 void FindSmallBones(
-	FGeometryCollection& Collection,
+	const FGeometryCollection& Collection,
 	const TArrayView<const int32>& TransformIndices,
 	const TArrayView<const double>& Volumes,
 	double MinVolume,
-	TArray<int32>& OutSmallBones
+	TArray<int32>& OutSmallBones,
+	bool bIncludeClusters
 )
 {
 	FilterBonesByVolume(Collection, TransformIndices, Volumes, 
@@ -1326,7 +1303,7 @@ void FindSmallBones(
 		{
 			return Volume < MinVolume;
 		},
-		OutSmallBones);
+		OutSmallBones, bIncludeClusters);
 }
 
 
@@ -1677,6 +1654,256 @@ int32 MergeBones(
 	return INDEX_NONE; // TODO: consider tracking smallest index of updated groups?  but no reason to do so currently
 }
 
+namespace
+{
+	static double GetClusterVolumeFromRigidsHelper(const FGeometryCollection& Collection, const TArrayView<const double>& Volumes, int32 BoneIdx)
+	{
+		if (Collection.SimulationType[BoneIdx] == FGeometryCollection::ESimulationTypes::FST_Rigid)
+		{
+			return Volumes[BoneIdx];
+		}
+		double Sum = 0;
+		for (int32 Child : Collection.Children[BoneIdx])
+		{
+			Sum += GetClusterVolumeFromRigidsHelper(Collection, Volumes, Child);
+		}
+		return Sum;
+	}
+
+	static FBox GetClusterBoundsFromRigidsHelper(const FGeometryCollection& Collection, const TArray<FBox>& GeoBoxesInWorldSpace, int32 BoneIdx)
+	{
+		int32 GeoIdx = Collection.TransformToGeometryIndex[BoneIdx];
+		if (GeoIdx != INDEX_NONE)
+		{
+			return GeoBoxesInWorldSpace[GeoIdx];
+		}
+		FBox Combined(ForceInit);
+		for (int32 Child : Collection.Children[BoneIdx])
+		{
+			Combined += GetClusterBoundsFromRigidsHelper(Collection, GeoBoxesInWorldSpace, Child);
+		}
+		return Combined;
+	}
+}
+
+void MergeClusters(
+	FGeometryCollection& Collection,
+	const TArrayView<const double>& Volumes,
+	double MinVolume,
+	const TArrayView<const int32>& SmallTransformIndices,
+	UE::PlanarCut::ENeighborSelectionMethod NeighborSelectionMethod,
+	bool bOnlyMergeInProximity,
+	bool bOnlySameParent
+)
+{
+	FGeometryCollectionProximityUtility ProximityUtility(&Collection);
+	ProximityUtility.RequireProximity();
+	const TManagedArray<TSet<int32>>& Proximity = Collection.GetAttribute<TSet<int32>>("Proximity", FGeometryCollection::GeometryGroup);
+	const TManagedArray<int32>* Level = Collection.FindAttribute<int32>("Level", FGeometryCollection::TransformGroup);
+	if (!ensure(Level != nullptr))
+	{
+		// cluster merging requires a level attribute
+		return;
+	}
+
+	TArray<FBox> GeoBoxesInWorldSpace;
+	TArray<FTransform> GlobalTransforms;
+	GeometryCollectionAlgo::GlobalMatrices(Collection.Transform, Collection.Parent, GlobalTransforms);
+	GeoBoxesInWorldSpace.SetNum(Collection.TransformIndex.Num());
+	for (int32 GeoIdx = 0; GeoIdx < GeoBoxesInWorldSpace.Num(); ++GeoIdx)
+	{
+		FTransform ToWorld = GlobalTransforms[Collection.TransformIndex[GeoIdx]];
+		FBox GeoBox(ForceInit);
+		int32 VStart = Collection.VertexStart[GeoIdx];
+		int32 VEnd = VStart + Collection.VertexCount[GeoIdx];
+		for (int32 VIdx = VStart; VIdx < VEnd; ++VIdx)
+		{
+			GeoBox += ToWorld.TransformPosition((FVector)Collection.Vertex[VIdx]);
+		}
+		GeoBoxesInWorldSpace[GeoIdx] = GeoBox;
+	}
+
+	// helper to grab same level neighbors from the proximity info + traversal
+	// only selects rigid or cluster bones
+	// Note we do the traversal each time rather using a pre-computed connection graph because the cluster neighbors will change as we merge clusters
+	// (the Proximity attribute should be stable because it is on the geometry group, not the transform group)
+	auto GetSameLevelNeighbors = [&Collection, &Proximity, &Level, bOnlySameParent](int32 BoneIdx, TSet<int32>& OutNeighbors) -> void
+	{
+		int32 TargetLevel = (*Level)[BoneIdx];
+		TArray<int32> ToProcess;
+		ToProcess.Add(BoneIdx);
+		int32 BoneParent = Collection.Parent[BoneIdx];
+		while (!ToProcess.IsEmpty())
+		{
+			int32 ProcessIdx = ToProcess.Pop();
+			// we're processing a rigid node: access the proximity graph and look for neighbor nodes
+			if (Collection.SimulationType[ProcessIdx] == FGeometryCollection::ESimulationTypes::FST_Rigid
+				&& Collection.TransformToGeometryIndex[ProcessIdx] != INDEX_NONE)
+			{
+				for (int32 NbrGeo : Proximity[Collection.TransformToGeometryIndex[ProcessIdx]])
+				{
+					// Traverse parents to find a same-level cluster on this neighbor
+					int32 TraverseParent = Collection.TransformIndex[NbrGeo];
+					while (TraverseParent != -1 && (*Level)[TraverseParent] > TargetLevel)
+					{
+						TraverseParent = Collection.Parent[TraverseParent];
+					}
+					if (TraverseParent == -1) // failed to find a same-level neighbor
+					{
+						continue;
+					}
+					// any different, same-level cluster is a valid neighbor to consider
+					bool bIsRigid = Collection.SimulationType[TraverseParent] == FGeometryCollection::ESimulationTypes::FST_Rigid;
+					// A valid cluster is one that has not been emptied by a previous merge (though we should not reach such an empty cluster in this traversal)
+					bool bIsValidCluster = !Collection.Children[TraverseParent].IsEmpty() && Collection.SimulationType[TraverseParent] == FGeometryCollection::ESimulationTypes::FST_Clustered;
+					bool bValidParent = !bOnlySameParent || Collection.Parent[TraverseParent] == BoneParent;
+					if (TraverseParent != BoneIdx && TraverseParent != -1 && (*Level)[TraverseParent] == TargetLevel && bValidParent
+						&& (bIsRigid || bIsValidCluster)) // only merge to rigid or cluster bones
+					{
+						OutNeighbors.Add(TraverseParent);
+					}
+				}
+			}
+			// we're processing a cluster node: traverse down to children to find a rigid node with proximity info
+			else if (Collection.SimulationType[ProcessIdx] == FGeometryCollection::ESimulationTypes::FST_Clustered)
+			{
+				ToProcess.Reserve(ToProcess.Num() + Collection.Children[ProcessIdx].Num());
+				for (int32 ChildIdx : Collection.Children[ProcessIdx])
+				{
+					ToProcess.Add(ChildIdx);
+				}
+			}
+		}
+	};
+
+	TSet<int32> ReconsiderIfSmall, DoNotReconsider; // Sets to track nodes we've already merged, and may not need to re-merge
+	TSet<int32> Neighbors; // Set to be re-used in the below loop, tracking neighbors to consider
+	TArray<int32> ChildNodes; // Array to be re-used in the below loop, tracking nodes to merge under a new parent
+	TArray<int32> ToProcess(SmallTransformIndices); // Need a mutable copy of SmallTransformIndices to update the array with merged nodes
+	for (int32 ProcessIdx = 0; ProcessIdx < ToProcess.Num(); ++ProcessIdx)
+	{
+		int32 MergeIdx = SmallTransformIndices[ProcessIdx];
+		if (Collection.Parent[MergeIdx] == -1 || DoNotReconsider.Contains(MergeIdx)) // can't merge root nodes
+		{
+			continue;
+		}
+		bool bIsCluster = Collection.SimulationType[MergeIdx] == FGeometryCollection::ESimulationTypes::FST_Clustered;
+		if (ReconsiderIfSmall.Contains(MergeIdx))
+		{
+			double Volume = GetClusterVolumeFromRigidsHelper(Collection, Volumes, MergeIdx);
+			if (Volume > MinVolume)
+			{
+				continue;
+			}
+		}
+		Neighbors.Reset();
+		GetSameLevelNeighbors(MergeIdx, Neighbors);
+		int32 BestNeighbor = INDEX_NONE;
+		if (Neighbors.IsEmpty() && !bOnlyMergeInProximity)
+		{
+			FVector Center = GetClusterBoundsFromRigidsHelper(Collection, GeoBoxesInWorldSpace, MergeIdx).GetCenter();
+			double ClosestDistSq = FMathd::MaxReal;
+			for (int32 TransformIdx = 0; TransformIdx < Collection.Transform.Num(); ++TransformIdx)
+			{
+				if (TransformIdx == MergeIdx || (*Level)[MergeIdx] != (*Level)[TransformIdx]
+					|| Collection.SimulationType[TransformIdx] != FGeometryCollection::ESimulationTypes::FST_Clustered)
+				{
+					continue;
+				}
+				FVector OtherCenter = GetClusterBoundsFromRigidsHelper(Collection, GeoBoxesInWorldSpace, TransformIdx).GetCenter();
+				double DistSq = FVector::DistSquared(Center, OtherCenter);
+				if (BestNeighbor == INDEX_NONE || DistSq < ClosestDistSq)
+				{
+					BestNeighbor = TransformIdx;
+					ClosestDistSq = DistSq;
+				}
+			}
+		}
+		else
+		{
+			if (NeighborSelectionMethod == ENeighborSelectionMethod::NearestCenter)
+			{
+				FVector Center = GetClusterBoundsFromRigidsHelper(Collection, GeoBoxesInWorldSpace, MergeIdx).GetCenter();
+				double ClosestDistSq = FMathd::MaxReal;
+				for (int32 NeighborIdx : Neighbors)
+				{
+					FVector OtherCenter = GetClusterBoundsFromRigidsHelper(Collection, GeoBoxesInWorldSpace, NeighborIdx).GetCenter();
+					double DistSq = FVector::DistSquared(Center, OtherCenter);
+					if (BestNeighbor == INDEX_NONE || DistSq < ClosestDistSq)
+					{
+						BestNeighbor = NeighborIdx;
+						ClosestDistSq = DistSq;
+					}
+				}
+			}
+			else // ENeighborSelectionMethod::LargestNeighbor
+			{
+				double Volume = GetClusterVolumeFromRigidsHelper(Collection, Volumes, MergeIdx);
+				double LargestVolume = -1;
+				for (int32 NeighborIdx : Neighbors)
+				{
+					double NbrVolume = GetClusterVolumeFromRigidsHelper(Collection, Volumes, NeighborIdx);
+					if (BestNeighbor == INDEX_NONE || NbrVolume > LargestVolume)
+					{
+						BestNeighbor = NeighborIdx;
+						LargestVolume = NbrVolume;
+					}
+				}
+			}
+		}
+		if (BestNeighbor == INDEX_NONE)
+		{
+			continue;
+		}
+		
+		int32 ToMerge = MergeIdx;
+		if (Collection.SimulationType[BestNeighbor] != FGeometryCollection::ESimulationTypes::FST_Clustered)
+		{
+			if (Collection.SimulationType[ToMerge] == FGeometryCollection::ESimulationTypes::FST_Clustered)
+			{
+				Swap(ToMerge, BestNeighbor);
+			}
+			else
+			{
+				// both nodes are rigid: need to create a new cluster for them
+				ChildNodes.Reset();
+				ChildNodes.Add(MergeIdx);
+				ChildNodes.Add(BestNeighbor);
+				int32 NewCluster = FGeometryCollectionClusteringUtility::ClusterBonesUnderNewNode(&Collection, MergeIdx, ChildNodes, true, false);
+				// Now that the node is clustered at a lower level, we shouldn't try to re-merge it with any other nodes
+				DoNotReconsider.Add(BestNeighbor);
+				double Volume = GetClusterVolumeFromRigidsHelper(Collection, Volumes, NewCluster);
+				if (Volume < MinVolume)
+				{
+					// if the new cluster is still small, we may want to merge it again
+					ToProcess[ProcessIdx] = NewCluster;
+					ProcessIdx--; // this re-considers ProcessIdx on the next iteration
+				}
+				continue;
+			}
+		}
+
+		// From the if / swap above, BestNeighbor must now be a cluster node, so we can merge under it
+		checkSlow(Collection.SimulationType[BestNeighbor] == FGeometryCollection::ESimulationTypes::FST_Clustered);
+		int32 ParentNode = BestNeighbor;
+		ChildNodes.Reset();
+		if (!bIsCluster)
+		{
+			ChildNodes.Add(ToMerge);
+		}
+		else
+		{
+			ChildNodes.Append(Collection.Children[ToMerge].Array());
+		}
+		// Note: This updates the Level, Children and Parent attributes
+		FGeometryCollectionClusteringUtility::ClusterBonesUnderExistingNode(&Collection, ParentNode, ChildNodes);
+		ReconsiderIfSmall.Add(BestNeighbor);
+	}
+
+	FGeometryCollectionClusteringUtility::RemoveDanglingClusters(&Collection);
+	FGeometryCollectionClusteringUtility::RemoveClustersOfOnlyOneChild(&Collection);
+}
+
 namespace
 {
 	void AddAllChildrenExceptEmbedded(FGeometryCollection& Collection, TSet<int32>& NodesSet, int32 Node)
diff --git a/Engine/Plugins/Experimental/PlanarCutPlugin/Source/PlanarCut/Public/PlanarCut.h b/Engine/Plugins/Experimental/PlanarCutPlugin/Source/PlanarCut/Public/PlanarCut.h
--- a/Engine/Plugins/Experimental/PlanarCutPlugin/Source/PlanarCut/Public/PlanarCut.h
+++ b/Engine/Plugins/Experimental/PlanarCutPlugin/Source/PlanarCut/Public/PlanarCut.h
@@ -323,12 +323,14 @@ int32 PLANARCUT_API CutMultipleWithMultiplePlanes(
  * @param TransformIndices		The transform indices to process, or empty if all should be processed
  * @param OutVolumes			Output array, to be filled w/ volumes of geometry; 1:1 w/ TransformIndices array
  * @param ScalePerDimension		Scale to apply per dimension (e.g. 1/100 converts volume from centimeters^3 to meters^3)
+ * @param bIncludeClusters		Whether to set the volumes for clusters as a sum of the child volumes. Otherwise, cluster volumes are set to zero.
  */
 void PLANARCUT_API FindBoneVolumes(
 	FGeometryCollection& Collection,
 	const TArrayView<const int32>& TransformIndices,
 	TArray<double>& OutVolumes,
-	double ScalePerDimension = .01
+	double ScalePerDimension = .01,
+	bool bIncludeClusters = false
 );
 
 /**
@@ -339,13 +341,15 @@ void PLANARCUT_API FindBoneVolumes(
  * @param Volumes				Volumes of geometry; 1:1 w/ TransformIndices array
  * @param MinVolume				Geometry smaller than this quantity will be chosen
  * @param OutSmallBones			Output array, to be filled with transform indices for small pieces of geometry
+ * @param bIncludeClusters		Whether to include clusters in the search
  */
 void PLANARCUT_API FindSmallBones(
-	FGeometryCollection& Collection,
+	const FGeometryCollection& Collection,
 	const TArrayView<const int32>& TransformIndices,
 	const TArrayView<const double>& Volumes,
 	double MinVolume,
-	TArray<int32>& OutSmallBones
+	TArray<int32>& OutSmallBones,
+	bool bIncludeClusters = false
 );
 
 /**
@@ -356,13 +360,15 @@ void PLANARCUT_API FindSmallBones(
  * @param Volumes				Volumes of geometry; 1:1 w/ TransformIndices array
  * @param Filter				Geometry for which the volume filter returns true will be chosen
  * @param OutSmallBones			Output array, to be filled with transform indices for small pieces of geometry
+ * @param bIncludeClusters		Whether to allow clusters in the filter
  */
 void PLANARCUT_API FilterBonesByVolume(
-	FGeometryCollection& Collection,
+	const FGeometryCollection& Collection,
 	const TArrayView<const int32>& TransformIndices,
 	const TArrayView<const double>& Volumes,
 	TFunctionRef<bool(double Volume, int32 BoneIdx)> Filter,
-	TArray<int32>& OutSmallBones
+	TArray<int32>& OutSmallBones,
+	bool bIncludeClusters = false
 );
 
 namespace UE
@@ -399,6 +405,27 @@ int32 PLANARCUT_API MergeBones(
 	UE::PlanarCut::ENeighborSelectionMethod NeighborSelectionMethod
 );
 
+/**
+ * Merge chosen bones into neighboring clusters.
+ *
+ * @param Collection				The collection to be processed
+ * @param Volumes					Volumes of geometry; 1:1 w/ TransformIndices array
+ * @param MinVolume					If merged small geometry is larger than this, it will not require further merging
+ * @param SmallTransformIndices		Transform indices of pieces that we want to merge
+ * @param NeighborSelectionMethod	How to choose which neighbor to merge to
+ * @param bOnlyMergeInProximity		Whether to only follow proximity graph, or to fall back to nearest center when for when no proximity links are found
+ * @param bOnlySameParent			Whether to only merge nodes that have the same parent in the hierarchy
+ */
+void PLANARCUT_API MergeClusters(
+	FGeometryCollection& Collection,
+	const TArrayView<const double>& Volumes,
+	double MinVolume,
+	const TArrayView<const int32>& SmallTransformIndices,
+	UE::PlanarCut::ENeighborSelectionMethod NeighborSelectionMethod,
+	bool bOnlyMergeInProximity,
+	bool bOnlySameParent
+);
+
 /**
  * Merge all chosen nodes into a single node.  Unlike MergeBones(), does not account for proximity and will always only merge the selected bones.
  *
diff --git a/Engine/Source/Runtime/Experimental/Chaos/Private/GeometryCollection/GeometryCollectionClusteringUtility.cpp b/Engine/Source/Runtime/Experimental/Chaos/Private/GeometryCollection/GeometryCollectionClusteringUtility.cpp
--- a/Engine/Source/Runtime/Experimental/Chaos/Private/GeometryCollection/GeometryCollectionClusteringUtility.cpp
+++ b/Engine/Source/Runtime/Experimental/Chaos/Private/GeometryCollection/GeometryCollectionClusteringUtility.cpp
@@ -10,7 +10,7 @@
 static int32 ChaosValidateResultsOfEditOperations = 0;
 static FAutoConsoleVariableRef CVarChaosStillCheckDistanceThreshold(TEXT("p.fracture.ValidateResultsOfEditOperations"), ChaosValidateResultsOfEditOperations, TEXT("When on this will enable result validation for fracture tool edit operations (can be slow for large geometry collection) [def:0]"));
 
-void FGeometryCollectionClusteringUtility::ClusterBonesUnderNewNode(FGeometryCollection* GeometryCollection, const int32 InsertAtIndex, const TArray<int32>& SelectedBones, bool CalcNewLocalTransform, bool Validate)
+int32 FGeometryCollectionClusteringUtility::ClusterBonesUnderNewNode(FGeometryCollection* GeometryCollection, const int32 InsertAtIndex, const TArray<int32>& SelectedBones, bool CalcNewLocalTransform, bool Validate)
 {
 	check(GeometryCollection);
 
@@ -19,7 +19,7 @@ void FGeometryCollectionClusteringUtility::ClusterBonesUnderNewNode(FGeometryCol
 	return ClusterBonesUnderNewNodeWithParent(GeometryCollection, ParentIdx, SelectedBones, CalcNewLocalTransform, Validate);
 }
 
-void FGeometryCollectionClusteringUtility::ClusterBonesUnderNewNodeWithParent(FGeometryCollection* GeometryCollection, const int32 ParentOfNewNode, const TArray<int32>& SelectedBones, bool CalcNewLocalTransform, bool Validate)
+int32 FGeometryCollectionClusteringUtility::ClusterBonesUnderNewNodeWithParent(FGeometryCollection* GeometryCollection, const int32 ParentOfNewNode, const TArray<int32>& SelectedBones, bool CalcNewLocalTransform, bool Validate)
 {
 	check(GeometryCollection);
 
@@ -59,6 +59,8 @@ void FGeometryCollectionClusteringUtility::ClusterBonesUnderNewNodeWithParent(FG
 	{
 		ValidateResults(GeometryCollection);
 	}
+
+	return NewBoneIndex;
 }
 
 void FGeometryCollectionClusteringUtility::ClusterAllBonesUnderNewRoot(FGeometryCollection* GeometryCollection)
diff --git a/Engine/Source/Runtime/Experimental/Chaos/Public/GeometryCollection/GeometryCollectionClusteringUtility.h b/Engine/Source/Runtime/Experimental/Chaos/Public/GeometryCollection/GeometryCollectionClusteringUtility.h
--- a/Engine/Source/Runtime/Experimental/Chaos/Public/GeometryCollection/GeometryCollectionClusteringUtility.h
+++ b/Engine/Source/Runtime/Experimental/Chaos/Public/GeometryCollection/GeometryCollectionClusteringUtility.h
@@ -34,11 +34,13 @@ public:
 	*   L2    A  B
 	*
 	* Node E has no geometry of its own, only a transform by which to control A & B as a single unit
+	* 
+	* @return Index of the New Node
 	*/
-	static void ClusterBonesUnderNewNode(FGeometryCollection* GeometryCollection, const int32 InsertAtIndex, const TArray<int32>& SelectedBones, bool CalcNewLocalTransform, bool Validate = true);
+	static int32 ClusterBonesUnderNewNode(FGeometryCollection* GeometryCollection, const int32 InsertAtIndex, const TArray<int32>& SelectedBones, bool CalcNewLocalTransform, bool Validate = true);
 
 	// Same as ClusterBonesUnderNewNode, but specify the parent of the new node instead of a sibling
-	static void ClusterBonesUnderNewNodeWithParent(FGeometryCollection* GeometryCollection, const int32 ParentOfNewNode, const TArray<int32>& SelectedBones, bool CalcNewLocalTransform, bool Validate = true);
+	static int32 ClusterBonesUnderNewNodeWithParent(FGeometryCollection* GeometryCollection, const int32 ParentOfNewNode, const TArray<int32>& SelectedBones, bool CalcNewLocalTransform, bool Validate = true);
 
 	/** Cluster all existing bones under a new root node, so there is now only one root node and a completely flat hierarchy underneath it */
 	static void ClusterAllBonesUnderNewRoot(FGeometryCollection* GeometryCollection);
