make geometry collection fracture also track vertex color alpha channels (by using a custom per-vertex attribute instead of the built-in FVector3f dynamic mesh vertex colors)

diff --git a/Engine/Plugins/Experimental/PlanarCutPlugin/Source/PlanarCut/Private/GeometryMeshConversion.cpp b/Engine/Plugins/Experimental/PlanarCutPlugin/Source/PlanarCut/Private/GeometryMeshConversion.cpp
--- a/Engine/Plugins/Experimental/PlanarCutPlugin/Source/PlanarCut/Private/GeometryMeshConversion.cpp
+++ b/Engine/Plugins/Experimental/PlanarCutPlugin/Source/PlanarCut/Private/GeometryMeshConversion.cpp
@@ -67,9 +67,10 @@ namespace AugmentedDynamicMesh
 {
 	// An invalid color, to be replaced by neighboring valid colors (or DefaultVertexColor, if neighboring colors were not found)
 	// Use a large negative value as a clear unset / invalid value, but do not go all the way to -MaxReal (or overflow will break things)
-	const static FVector3f UnsetVertexColor = FVector3f(-FMathf::MaxReal * .25f, -FMathf::MaxReal * .25f, -FMathf::MaxReal * .25f);
-	const static FVector3f DefaultVertexColor = FVector3f::ZeroVector;
+	const static FVector4f UnsetVertexColor = FVector4f(-FMathf::MaxReal * .25f, -FMathf::MaxReal * .25f, -FMathf::MaxReal * .25f, -FMathf::MaxReal * .25f);
+	const static FVector4f DefaultVertexColor = FVector4f(0, 0, 0, 1);
 
+	FName ColorAttribName = "ColorAttrib";
 	FName TangentUAttribName = "TangentUAttrib";
 	FName TangentVAttribName = "TangentVAttrib";
 	FName VisibleAttribName = "VisibleAttrib";
@@ -131,10 +132,10 @@ namespace AugmentedDynamicMesh
 
 	void Augment(FDynamicMesh3& Mesh, int32 NumUVChannels)
 	{
-		Mesh.EnableVertexColors(UnsetVertexColor);
 		Mesh.EnableVertexNormals(FVector3f::UnitZ());
 		Mesh.EnableAttributes();
 		Mesh.Attributes()->EnableMaterialID();
+		Mesh.Attributes()->AttachAttribute(ColorAttribName, new TDynamicMeshVertexAttribute<float, 4>(&Mesh));
 		Mesh.Attributes()->AttachAttribute(TangentUAttribName, new TDynamicMeshVertexAttribute<float, 3>(&Mesh));
 		Mesh.Attributes()->AttachAttribute(TangentVAttribName, new TDynamicMeshVertexAttribute<float, 3>(&Mesh));
 		TDynamicMeshScalarTriangleAttribute<bool>* VisAttrib = new TDynamicMeshScalarTriangleAttribute<bool>(&Mesh);
@@ -147,20 +148,27 @@ namespace AugmentedDynamicMesh
 		EnableUVChannels(Mesh, NumUVChannels);
 	}
 
+	void AddVertexColorAttribute(FDynamicMesh3& Mesh)
+	{
+		Mesh.Attributes()->AttachAttribute(ColorAttribName, new TDynamicMeshVertexAttribute<float, 4>(&Mesh));
+	}
+
 	bool IsAugmented(const FDynamicMesh3& Mesh)
 	{
 		return Mesh.HasAttributes()
+			&& Mesh.Attributes()->HasAttachedAttribute(ColorAttribName)
 			&& Mesh.Attributes()->HasAttachedAttribute(TangentUAttribName)
 			&& Mesh.Attributes()->HasAttachedAttribute(TangentVAttribName)
 			&& Mesh.Attributes()->HasAttachedAttribute(VisibleAttribName)
 			&& Mesh.Attributes()->HasMaterialID()
-			&& Mesh.HasVertexColors()
 			&& Mesh.HasVertexNormals();
 	}
 
 	void SetDefaultAttributes(FDynamicMesh3& Mesh, bool bGlobalVisibility)
 	{
 		checkSlow(IsAugmented(Mesh));
+		TDynamicMeshVertexAttribute<float, 4>* Colors =
+			static_cast<TDynamicMeshVertexAttribute<float, 4>*>(Mesh.Attributes()->GetAttachedAttribute(ColorAttribName));
 		TDynamicMeshVertexAttribute<float, 3>* Us =
 			static_cast<TDynamicMeshVertexAttribute<float, 3>*>(Mesh.Attributes()->GetAttachedAttribute(TangentUAttribName));
 		TDynamicMeshVertexAttribute<float, 3>* Vs =
@@ -173,6 +181,7 @@ namespace AugmentedDynamicMesh
 			VectorUtil::MakePerpVectors(N, U, V);
 			Us->SetValue(VID, U);
 			Vs->SetValue(VID, V);
+			Colors->SetValue(VID, UnsetVertexColor);
 		}
 
 		TDynamicMeshScalarTriangleAttribute<bool>* Visible =
@@ -284,6 +293,24 @@ namespace AugmentedDynamicMesh
 		Vs->GetValue(VID, V);
 	}
 
+	FVector4f GetVertexColor(const FDynamicMesh3& Mesh, int VID)
+	{
+		checkSlow(IsAugmented(Mesh));
+		const TDynamicMeshVertexAttribute<float, 4>* Colors =
+			static_cast<const TDynamicMeshVertexAttribute<float, 4>*>(Mesh.Attributes()->GetAttachedAttribute(ColorAttribName));
+		FVector4f Color;
+		Colors->GetValue(VID, Color);
+		return Color;
+	}
+
+	void SetVertexColor(FDynamicMesh3& Mesh, int VID, FVector4f Color)
+	{
+		checkSlow(IsAugmented(Mesh));
+		TDynamicMeshVertexAttribute<float, 4>* Colors =
+			static_cast<TDynamicMeshVertexAttribute<float, 4>*>(Mesh.Attributes()->GetAttachedAttribute(ColorAttribName));
+		Colors->SetValue(VID, Color);
+	}
+
 	template<typename FOverlay, typename VecType>
 	void SplitOverlayHelper(UE::Geometry::FDynamicMesh3& Mesh, FOverlay* Overlay)
 	{
@@ -713,9 +740,10 @@ namespace AugmentedDynamicMesh
 							if (VIDDist.Key == -1)
 							{
 								// no point within radius; can add a sample here
-								FVertexInfo Info(SamplePos, Normal, AugmentedDynamicMesh::DefaultVertexColor);
+								FVertexInfo Info(SamplePos, Normal);
 
 								int AddedVID = Mesh.AppendVertex(Info);
+								SetVertexColor(Mesh, AddedVID, DefaultVertexColor);
 								KnownSamples[ComponentIdx].InsertPointUnsafe(AddedVID, SamplePos);
 							}
 						}
@@ -1110,10 +1138,10 @@ namespace AugmentedDynamicMesh
 			}
 
 			// Copy color + tangents from the reference vertex (TODO: consider an average instead?)
-			FVector3f HoleColor;
+			FVector4f HoleColor;
 			FVector3f HoleTanU, HoleTanV;
 			GetTangent(Mesh, HoleEdgeV.A, HoleTanU, HoleTanV);
-			HoleColor = Mesh.GetVertexColor(HoleEdgeV.A);
+			HoleColor = GetVertexColor(Mesh, HoleEdgeV.A);
 			
 			int32 PrevMaxVID = Mesh.MaxVertexID();
 			// Create new vertices for all that weren't already assigned to an existing vertex
@@ -1127,8 +1155,9 @@ namespace AugmentedDynamicMesh
 
 				int32 V = AllBoundaryVerts[Idx];
 				FVector3d P = Mesh.GetVertex(V);
-				UE::Geometry::FVertexInfo Info(P, (FVector3f)HoleNormals[CurBdry], HoleColor);
+				UE::Geometry::FVertexInfo Info(P, (FVector3f)HoleNormals[CurBdry]);
 				int32 NewV = Mesh.AppendVertex(Info);
+				SetVertexColor(Mesh, NewV, HoleColor);
 				SetTangent(Mesh, NewV, (FVector3f)HoleNormals[CurBdry], HoleTanU, HoleTanV);
 				for (int32 UVIdx = 0; UVIdx < NumUVs; UVIdx++)
 				{
@@ -1215,7 +1244,7 @@ namespace
 
 		auto IsUnset = [](const FLinearColor& Color) -> bool
 		{
-			return Color.R < 0 || Color.G < 0 || Color.B < 0;
+			return Color.R < 0 || Color.G < 0 || Color.B < 0 || Color.A < 0;
 		};
 		constexpr float CopyColorDistance = 1e-03;
 
@@ -1266,7 +1295,7 @@ namespace
 			TArray<FLink> Links; // connections between vertices -- note: always symmetric, duplicates are allowed
 			Links.Reserve(NumV * 3);
 
-			TArray<FVector3f> FixedColors; // Note: FLinearColor has an alpha but FDynamicMesh3 currently only tracks rgb per vertex, so we only solve for rgb
+			TArray<FVector4f> FixedColors;
 			TArray<float> FixedColorWts;
 			FixedColors.SetNumZeroed(NumV);
 			FixedColorWts.SetNumZeroed(NumV);
@@ -1368,7 +1397,7 @@ namespace
 							}
 							CanSolveComponents.Add(Component);
 							int32 SharedVtx = VtxComps.GetComponent(NbrIdx - StartV);
-							FixedColors[SharedVtx] += FVector3f(Color.R, Color.G, Color.B);
+							FixedColors[SharedVtx] += FVector4f(Color.R, Color.G, Color.B, Color.A);
 							FixedColorWts[SharedVtx] += 1.0f;
 						}
 					}
@@ -1442,9 +1471,9 @@ namespace
 				}
 
 				FSparseMatf SparseMatrix(NumToSolve, NumToSolve);
-				Eigen::VectorXf X[3]{ Eigen::VectorXf(NumToSolve), Eigen::VectorXf(NumToSolve), Eigen::VectorXf(NumToSolve) };
-				Eigen::VectorXf B[3]{ Eigen::VectorXf(NumToSolve), Eigen::VectorXf(NumToSolve), Eigen::VectorXf(NumToSolve) };
-				for (int32 SubIdx = 0; SubIdx < 3; ++SubIdx)
+				Eigen::VectorXf X[4]{ Eigen::VectorXf(NumToSolve), Eigen::VectorXf(NumToSolve), Eigen::VectorXf(NumToSolve), Eigen::VectorXf(NumToSolve) };
+				Eigen::VectorXf B[4]{ Eigen::VectorXf(NumToSolve), Eigen::VectorXf(NumToSolve), Eigen::VectorXf(NumToSolve), Eigen::VectorXf(NumToSolve) };
+				for (int32 SubIdx = 0; SubIdx < 4; ++SubIdx)
 				{
 					B[SubIdx].setZero();
 				}
@@ -1471,8 +1500,8 @@ namespace
 					{
 						if (FixedColorWts[SharedB] > 0)
 						{
-							FVector3f BVal = FixedColors[SharedB] * Link.Wt;
-							for (int32 SubIdx = 0; SubIdx < 3; ++SubIdx)
+							FVector4f BVal = FixedColors[SharedB] * Link.Wt;
+							for (int32 SubIdx = 0; SubIdx < 4; ++SubIdx)
 							{
 								B[SubIdx][LocalA] += BVal[SubIdx];
 							}
@@ -1501,7 +1530,7 @@ namespace
 				MatrixSolver.analyzePattern(SparseMatrix);
 				MatrixSolver.factorize(SparseMatrix);
 
-				ParallelFor(3, [&](int32 Idx)
+				ParallelFor(4, [&](int32 Idx)
 					{
 						X[Idx] = MatrixSolver.solve(B[Idx]);
 					});
@@ -1515,8 +1544,8 @@ namespace
 					int32 CompIdx = ToComponentMap[SharedIdx];
 					if (CompIdx != INDEX_NONE)
 					{
-						FVector3f SolvedColor;
-						for (int32 SubIdx = 0; SubIdx < 3; ++SubIdx)
+						FVector4f SolvedColor;
+						for (int32 SubIdx = 0; SubIdx < 4; ++SubIdx)
 						{
 							// Note: make sure the solved color is non-negative, so solver error cannot make the color 'unset'
 							SolvedColor[SubIdx] = FMath::Max(X[SubIdx][CompIdx], 0.f);
@@ -1578,7 +1607,7 @@ FCellMeshes::FCellMeshes(int32 NumUVLayersIn, const FDynamicMesh3& SingleCutter,
 	// Make sure color is unset
 	for (int VID : CellMeshes[0].AugMesh.VertexIndicesItr())
 	{
-		CellMeshes[0].AugMesh.SetVertexColor(VID, AugmentedDynamicMesh::UnsetVertexColor);
+		AugmentedDynamicMesh::SetVertexColor(CellMeshes[0].AugMesh, VID, AugmentedDynamicMesh::UnsetVertexColor);
 	}
 
 	// first mesh is the same as the second mesh, but will be subtracted b/c it's the "outside cell"
@@ -1878,7 +1907,6 @@ void FCellMeshes::CreateMeshesForBoundedPlanesWithoutNoise(int NumCells, const F
 		PlaneVertInfo.bHaveC = true;
 		PlaneVertInfo.bHaveUV = false;
 		PlaneVertInfo.bHaveN = true;
-		PlaneVertInfo.Color = AugmentedDynamicMesh::UnsetVertexColor;
 		int VertStart[2]{ -1, -1 };
 		for (int MeshIdx = 0; MeshIdx < NumMeshes; MeshIdx++)
 		{
@@ -1901,6 +1929,7 @@ void FCellMeshes::CreateMeshesForBoundedPlanesWithoutNoise(int NumCells, const F
 				PlaneVertInfo.Position = FVector3d(Cells.PlaneBoundaryVertices[BoundaryVertex]);
 				FVector2f UV = (FVector2f(PlaneFrame.ToPlaneUV(PlaneVertInfo.Position)) - MinUV) * static_cast<float>(GlobalUVScale);
 				int VID = Meshes[MeshIdx]->AppendVertex(PlaneVertInfo);
+				AugmentedDynamicMesh::SetVertexColor(*Meshes[MeshIdx], VID, AugmentedDynamicMesh::UnsetVertexColor);
 				AugmentedDynamicMesh::SetAllUV(*Meshes[MeshIdx], VID, UV, NumUVLayers);
 			}
 		}
@@ -1991,7 +2020,6 @@ void FCellMeshes::CreateMeshesForBoundedPlanesWithNoise(int NumCells, const FPla
 	{
 		AugmentedDynamicMesh::EnableUVChannels(PlaneMesh, NumUVLayers);
 		PlaneMesh.EnableVertexNormals(FVector3f::UnitZ());
-		PlaneMesh.EnableVertexColors(AugmentedDynamicMesh::UnsetVertexColor);
 		PlaneMesh.EnableAttributes();
 		PlaneMesh.Attributes()->EnableMaterialID();
 		PlaneMesh.Attributes()->AttachAttribute(OriginalPositionAttribute, new TDynamicMeshVertexAttribute<double, 3>(&PlaneMesh));
@@ -2048,7 +2076,7 @@ void FCellMeshes::CreateMeshesForBoundedPlanesWithNoise(int NumCells, const FPla
 			PlaneVertInfo.bHaveN = true;
 			PlaneVertInfo.Normal = Normal;
 			// UVs will be set below, after noise is added
-			PlaneVertInfo.Color = AugmentedDynamicMesh::UnsetVertexColor;
+			// UnsetVertexColor will be set below
 
 			FPolygon2f Polygon;
 			for (int BoundaryVertex : PlaneBoundary)
@@ -2272,11 +2300,16 @@ void FCellMeshes::CreateMeshesForBoundedPlanesWithNoise(int NumCells, const FPla
 	}
 
 	// clear "original position" attribute now that we have removed self-intersections
+	// and set unset vertex colors
 	for (int CellIdx = 0; CellIdx < NumCells; CellIdx++)
 	{
 		FCellInfo& CellInfo = CellMeshes[CellIdx];
 		FDynamicMesh3& Mesh = CellInfo.AugMesh;
 		Mesh.Attributes()->RemoveAttribute(OriginalPositionAttribute);
+		for (int VID : Mesh.VertexIndicesItr())
+		{
+			AugmentedDynamicMesh::SetVertexColor(Mesh, VID, AugmentedDynamicMesh::UnsetVertexColor);
+		}
 	}
 
 	// recompute UVs using new positions after noise was applied + fixed
@@ -2360,7 +2393,6 @@ void FCellMeshes::CreateMeshesForSinglePlane(const FPlanarCells& Cells, const FA
 	FVertexInfo PlaneVertInfo;
 	PlaneVertInfo.bHaveC = true;
 	PlaneVertInfo.bHaveN = true;
-	PlaneVertInfo.Color = AugmentedDynamicMesh::UnsetVertexColor;
 	PlaneVertInfo.Normal = -FVector3f(Plane.GetNormal());
 
 	for (int CornerIdx = 0; CornerIdx < 4; CornerIdx++)
@@ -2538,6 +2570,7 @@ void FDynamicMeshCollection::Init(const FGeometryCollection* Collection, const T
 			VertexInfo.Color = FVector3f(Collection->Color[Idx]);
 			VertexInfo.Normal = (FVector3f)CollectionToLocal.TransformVectorNoScale(FVector3d(Collection->Normal[Idx]));
 			int VID = Mesh.AppendVertex(VertexInfo);
+			AugmentedDynamicMesh::SetVertexColor(Mesh, VID, FVector4f(Collection->Color[Idx]));
 			AugmentedDynamicMesh::SetTangent(Mesh, VID, VertexInfo.Normal,
 				(FVector3f)CollectionToLocal.TransformVectorNoScale(FVector3d(Collection->TangentU[Idx])),
 				(FVector3f)CollectionToLocal.TransformVectorNoScale(FVector3d(Collection->TangentV[Idx])));
@@ -3274,7 +3307,7 @@ bool FDynamicMeshCollection::UpdateCollection(const FTransform& FromCollection,
 		AugmentedDynamicMesh::GetTangent(Mesh, VID, TangentU, TangentV);
 		Output.TangentU[CopyToIdx] = (FVector3f)FromCollection.InverseTransformVectorNoScale(FVector(TangentU));
 		Output.TangentV[CopyToIdx] = (FVector3f)FromCollection.InverseTransformVectorNoScale(FVector(TangentV));
-		Output.Color[CopyToIdx] = FLinearColor(FVector(Mesh.GetVertexColor(VID)));
+		Output.Color[CopyToIdx] = FLinearColor(AugmentedDynamicMesh::GetVertexColor(Mesh, VID));
 
 		// Bone map is set based on the transform of the new geometry
 		Output.BoneMap[CopyToIdx] = TransformIdx;
@@ -3380,7 +3413,7 @@ int32 FDynamicMeshCollection::AppendToCollection(const FTransform& FromCollectio
 		AugmentedDynamicMesh::GetTangent(Mesh, VID, TangentU, TangentV);
 		Output.TangentU[CopyToIdx] = (FVector3f)FromCollection.InverseTransformVectorNoScale(FVector(TangentU));
 		Output.TangentV[CopyToIdx] = (FVector3f)FromCollection.InverseTransformVectorNoScale(FVector(TangentV));
-		Output.Color[CopyToIdx] = FLinearColor(FVector(Mesh.GetVertexColor(VID)));
+		Output.Color[CopyToIdx] = FLinearColor(AugmentedDynamicMesh::GetVertexColor(Mesh, VID));
 
 		// Bone map is set based on the transform of the new geometry
 		Output.BoneMap[CopyToIdx] = TransformIdx;
