fix infinite clustering bug in geometry collection merge neighboring clusters algorithm

diff --git a/Engine/Plugins/Experimental/PlanarCutPlugin/Source/PlanarCut/Private/PlanarCut.cpp b/Engine/Plugins/Experimental/PlanarCutPlugin/Source/PlanarCut/Private/PlanarCut.cpp
--- a/Engine/Plugins/Experimental/PlanarCutPlugin/Source/PlanarCut/Private/PlanarCut.cpp
+++ b/Engine/Plugins/Experimental/PlanarCutPlugin/Source/PlanarCut/Private/PlanarCut.cpp
@@ -1782,7 +1782,7 @@ void MergeClusters(
 	TArray<int32> ToProcess(SmallTransformIndices); // Need a mutable copy of SmallTransformIndices to update the array with merged nodes
 	for (int32 ProcessIdx = 0; ProcessIdx < ToProcess.Num(); ++ProcessIdx)
 	{
-		int32 MergeIdx = SmallTransformIndices[ProcessIdx];
+		int32 MergeIdx = ToProcess[ProcessIdx];
 		if (Collection.Parent[MergeIdx] == -1 || DoNotReconsider.Contains(MergeIdx)) // can't merge root nodes
 		{
 			continue;
