Add support for geometry script conversion from geometry collection to dynamic mesh, via exposing geometry collection -> dynamic mesh conversion code in planarcut plugin and adding necessary dependencies for geometry script to call planarcut and geometry collection code.

Also make some related infrastructure accept tranforms in TArrayView form rather than TManagedArrayCollection so it's easier to pass through the active component transforms (which are passed back from the component via a tarray instead of a managed array)

diff --git a/Engine/Plugins/Experimental/GeometryScripting/GeometryScripting.uplugin b/Engine/Plugins/Experimental/GeometryScripting/GeometryScripting.uplugin
--- a/Engine/Plugins/Experimental/GeometryScripting/GeometryScripting.uplugin
+++ b/Engine/Plugins/Experimental/GeometryScripting/GeometryScripting.uplugin
@@ -34,6 +34,10 @@
 		{
 			"Name": "MeshModelingToolset",
 			"Enabled": true
+		},
+		{
+			"Name": "PlanarCut",
+			"Enabled":  true
 		}
 	]
 }
diff --git a/Engine/Plugins/Experimental/GeometryScripting/Source/GeometryScriptingCore/GeometryScriptingCore.Build.cs b/Engine/Plugins/Experimental/GeometryScripting/Source/GeometryScriptingCore/GeometryScriptingCore.Build.cs
--- a/Engine/Plugins/Experimental/GeometryScripting/Source/GeometryScriptingCore/GeometryScriptingCore.Build.cs
+++ b/Engine/Plugins/Experimental/GeometryScripting/Source/GeometryScriptingCore/GeometryScriptingCore.Build.cs
@@ -47,7 +47,10 @@ public class GeometryScriptingCore : ModuleRules
 				"MeshConversionEngineTypes",
 				"GeometryAlgorithms",
 				"ModelingOperators",
-				"ModelingComponents"
+				"ModelingComponents",
+				"PlanarCut",
+				"Chaos",
+				"GeometryCollectionEngine"
 			}
 			);
 
diff --git a/Engine/Plugins/Experimental/GeometryScripting/Source/GeometryScriptingCore/Private/SceneUtilityFunctions.cpp b/Engine/Plugins/Experimental/GeometryScripting/Source/GeometryScriptingCore/Private/SceneUtilityFunctions.cpp
--- a/Engine/Plugins/Experimental/GeometryScripting/Source/GeometryScriptingCore/Private/SceneUtilityFunctions.cpp
+++ b/Engine/Plugins/Experimental/GeometryScripting/Source/GeometryScriptingCore/Private/SceneUtilityFunctions.cpp
@@ -21,7 +21,10 @@
 #include "ConversionUtils/VolumeToDynamicMesh.h"
 #include "ConversionUtils/SkinnedMeshToDynamicMesh.h"
 #include "ConversionUtils/SplineComponentDeformDynamicMesh.h"
+#include "GeometryCollection/GeometryCollectionComponent.h"
+#include "GeometryCollection/GeometryCollectionObject.h"
 #include "Physics/ComponentCollisionUtil.h"
+#include "PlanarCut.h"
 
 #include UE_INLINE_GENERATED_CPP_BY_NAME(SceneUtilityFunctions)
 
@@ -194,6 +197,40 @@ UDynamicMesh* UGeometryScriptLibrary_SceneUtilityFunctions::CopyMeshFromComponen
 			UE::Geometry::AppendError(Debug, EGeometryScriptErrorType::InvalidInputs, LOCTEXT("CopyMeshFromComponent_InvalidBrushConversion", "CopyMeshFromComponent: BrushComponent conversion produced 0 triangles"));
 		}
 	}
+	else if (UGeometryCollectionComponent* GeometryCollectionComponent = Cast<UGeometryCollectionComponent>(Component))
+	{
+		LocalToWorld = GeometryCollectionComponent->GetComponentTransform();
+
+		if (const UGeometryCollection* RestCollection = GeometryCollectionComponent->GetRestCollection())
+		{
+			if (const FGeometryCollection* Collection = RestCollection->GetGeometryCollection().Get())
+			{
+				ToDynamicMesh->EditMesh([&](FDynamicMesh3& EditMesh)
+				{
+					FTransform UnusedTransform;
+					const TArray<FTransform3f>& DynamicTransforms = GeometryCollectionComponent->GetComponentSpaceTransforms3f();
+					if (!DynamicTransforms.IsEmpty())
+					{
+						ConvertGeometryCollectionToDynamicMesh(EditMesh, UnusedTransform, false, *Collection, true, DynamicTransforms, false, Collection->TransformIndex.GetConstArray());
+					}
+					else
+					{
+						ConvertGeometryCollectionToDynamicMesh(EditMesh, UnusedTransform, false, *Collection, true, TArrayView<const FTransform3f>(Collection->Transform.GetConstArray()), true, Collection->TransformIndex.GetConstArray());
+					}
+				},
+				EDynamicMeshChangeType::GeneralEdit, EDynamicMeshAttributeChangeFlags::Unknown, true);
+				Outcome = EGeometryScriptOutcomePins::Success;
+			}
+			else
+			{
+				UE::Geometry::AppendError(Debug, EGeometryScriptErrorType::InvalidInputs, LOCTEXT("CopyMeshFromComponent_MissingCollectionData", "CopyMeshFromComponent: GeometryCollectionComponent has null Geometry Collection data"));
+			}
+		}
+		else
+		{
+			UE::Geometry::AppendError(Debug, EGeometryScriptErrorType::InvalidInputs, LOCTEXT("CopyMeshFromComponent_MissingRestCollection", "CopyMeshFromComponent: GeometryCollectionComponent has null Rest Collection object"));
+		}
+	}
 
 	// transform mesh to world
 	if (Outcome == EGeometryScriptOutcomePins::Success && bTransformToWorld)
diff --git a/Engine/Plugins/Experimental/PlanarCutPlugin/Source/PlanarCut/Private/GeometryMeshConversion.cpp b/Engine/Plugins/Experimental/PlanarCutPlugin/Source/PlanarCut/Private/GeometryMeshConversion.cpp
--- a/Engine/Plugins/Experimental/PlanarCutPlugin/Source/PlanarCut/Private/GeometryMeshConversion.cpp
+++ b/Engine/Plugins/Experimental/PlanarCutPlugin/Source/PlanarCut/Private/GeometryMeshConversion.cpp
@@ -1579,6 +1579,19 @@ void SetGeometryCollectionAttributes(FDynamicMesh3& Mesh, int32 NumUVLayers)
 	AugmentedDynamicMesh::Augment(Mesh, NumUVLayers);
 }
 
+void ClearCustomGeometryCollectionAttributes(UE::Geometry::FDynamicMesh3& Mesh)
+{
+	Mesh.DiscardVertexNormals();
+	
+	Mesh.Attributes()->RemoveAttribute(AugmentedDynamicMesh::ColorAttribName);
+	Mesh.Attributes()->RemoveAttribute(AugmentedDynamicMesh::TangentUAttribName);
+	Mesh.Attributes()->RemoveAttribute(AugmentedDynamicMesh::TangentVAttribName);
+	Mesh.Attributes()->RemoveAttribute(AugmentedDynamicMesh::VisibleAttribName);
+	Mesh.Attributes()->RemoveAttribute(AugmentedDynamicMesh::InternalAttribName);
+
+	AugmentedDynamicMesh::EnableUVChannels(Mesh, 0, false, true); // set 0 UV channels to remove UV attributes
+}
+
 
 FCellMeshes::FCellMeshes(int32 NumUVLayersIn, FRandomStream& RandomStream, const FPlanarCells& Cells, FAxisAlignedBox3d DomainBounds, double Grout, double ExtendDomain, bool bIncludeOutsideCell)
 {
@@ -2529,7 +2542,7 @@ void FCellMeshes::CreateMeshesForSinglePlane(const FPlanarCells& Cells, const FA
 }
 
 template<typename TransformType>
-void FDynamicMeshCollection::InitTemplate(const FGeometryCollection* Collection, const TManagedArray<TransformType>& Transforms, const TArrayView<const int32>& TransformIndices, FTransform TransformCollection, bool bSaveIsolatedVertices)
+void FDynamicMeshCollection::InitTemplate(const FGeometryCollection* Collection, TArrayView<const TransformType> Transforms, const TArrayView<const int32>& TransformIndices, FTransform TransformCollection, bool bSaveIsolatedVertices)
 {
 	GeometryCollection::UV::FConstUVLayers UVLayers = GeometryCollection::UV::FindActiveUVLayers(*Collection);
 	int32 NumUVLayers = UVLayers.Num();
@@ -2545,7 +2558,15 @@ void FDynamicMeshCollection::InitTemplate(const FGeometryCollection* Collection,
 			continue;
 		}
 
-		FTransformSRT3d CollectionToLocal = FTransformSRT3d(GeometryCollectionAlgo::GlobalMatrix(Transforms, Collection->Parent, TransformIdx) * TransformCollection);
+		FTransformSRT3d CollectionToLocal;
+		if (bComponentSpaceTransforms)
+		{
+			CollectionToLocal = FTransformSRT3d(FTransform(Transforms[TransformIdx]) * TransformCollection);
+		}
+		else
+		{
+			CollectionToLocal = FTransformSRT3d(GeometryCollectionAlgo::GlobalMatrix(Transforms, TArrayView<const int32>(Collection->Parent.GetConstArray()), TransformIdx) * TransformCollection);
+		}
 
 		int32 AddedMeshIdx = Meshes.Add(new FMeshData(NumUVLayers));
 		FMeshData& MeshData = Meshes[AddedMeshIdx];
@@ -2583,6 +2604,10 @@ void FDynamicMeshCollection::InitTemplate(const FGeometryCollection* Collection,
 		FIntVector VertexOffset(VertexStart, VertexStart, VertexStart);
 		for (int32 Idx = Collection->FaceStart[GeometryIdx], N = Collection->FaceStart[GeometryIdx] + FaceCount; Idx < N; Idx++)
 		{
+			if (bSkipInvisible && !Collection->Visible[Idx])
+			{
+				continue;
+			}
 			FIndex3i AddTri = FIndex3i(Collection->Indices[Idx] - VertexOffset);
 			int TID = Mesh.AppendTriangle(AddTri, 0);
 			if (TID == FDynamicMesh3::NonManifoldID)
@@ -2631,12 +2656,12 @@ void FDynamicMeshCollection::InitTemplate(const FGeometryCollection* Collection,
 }
 
 
-void FDynamicMeshCollection::Init(const FGeometryCollection* Collection, const TManagedArray<FTransform>& Transforms, const TArrayView<const int32>& TransformIndices, FTransform TransformCollection, bool bSaveIsolatedVertices)
+void FDynamicMeshCollection::Init(const FGeometryCollection* Collection, TArrayView<const FTransform> Transforms, const TArrayView<const int32>& TransformIndices, FTransform TransformCollection, bool bSaveIsolatedVertices)
 {
 	InitTemplate(Collection, Transforms, TransformIndices, TransformCollection, bSaveIsolatedVertices);
 }
 
-void FDynamicMeshCollection::Init(const FGeometryCollection* Collection, const TManagedArray<FTransform3f>& Transforms, const TArrayView<const int32>& TransformIndices, FTransform TransformCollection, bool bSaveIsolatedVertices)
+void FDynamicMeshCollection::Init(const FGeometryCollection* Collection, TArrayView<const FTransform3f> Transforms, const TArrayView<const int32>& TransformIndices, FTransform TransformCollection, bool bSaveIsolatedVertices)
 {
 	InitTemplate(Collection, Transforms, TransformIndices, TransformCollection, bSaveIsolatedVertices);
 }
diff --git a/Engine/Plugins/Experimental/PlanarCutPlugin/Source/PlanarCut/Private/GeometryMeshConversion.h b/Engine/Plugins/Experimental/PlanarCutPlugin/Source/PlanarCut/Private/GeometryMeshConversion.h
--- a/Engine/Plugins/Experimental/PlanarCutPlugin/Source/PlanarCut/Private/GeometryMeshConversion.h
+++ b/Engine/Plugins/Experimental/PlanarCutPlugin/Source/PlanarCut/Private/GeometryMeshConversion.h
@@ -21,33 +21,39 @@ namespace PlanarCut
 /**
  * Add attributes necessary for a dynamic mesh to represent geometry from an FGeometryCollection
  */
-void PLANARCUT_API SetGeometryCollectionAttributes(UE::Geometry::FDynamicMesh3& Mesh, int32 NumUVLayers);
+void SetGeometryCollectionAttributes(UE::Geometry::FDynamicMesh3& Mesh, int32 NumUVLayers);
+
+/**
+ * Clear custom FGeometryCollection-specific attributes from a FDynamicMesh3
+ * Note: Does not remove the general attribute layer and MaterialID attributes, as these are not specific to geometry collections
+ */
+void ClearCustomGeometryCollectionAttributes(UE::Geometry::FDynamicMesh3& Mesh);
 
 // functions for DynamicMesh3 meshes that have FGeometryCollection attributes set
 namespace AugmentedDynamicMesh
 {
-	void PLANARCUT_API SetVisibility(UE::Geometry::FDynamicMesh3& Mesh, int TID, bool bIsVisible);
-	bool PLANARCUT_API GetVisibility(const UE::Geometry::FDynamicMesh3& Mesh, int TID);
-	void PLANARCUT_API SetInternal(UE::Geometry::FDynamicMesh3& Mesh, int TID, bool bIsInternal);
-	bool PLANARCUT_API GetInternal(const UE::Geometry::FDynamicMesh3& Mesh, int TID);
-	void PLANARCUT_API SetUV(UE::Geometry::FDynamicMesh3& Mesh, int VID, FVector2f UV, int UVLayer);
-	void PLANARCUT_API GetUV(const UE::Geometry::FDynamicMesh3& Mesh, int VID, FVector2f& UV, int UVLayer);
-	void PLANARCUT_API SetTangent(UE::Geometry::FDynamicMesh3& Mesh, int VID, FVector3f Normal, FVector3f TangentU, FVector3f TangentV);
-	void PLANARCUT_API GetTangent(const UE::Geometry::FDynamicMesh3& Mesh, int VID, FVector3f& U, FVector3f& V);
+	void SetVisibility(UE::Geometry::FDynamicMesh3& Mesh, int TID, bool bIsVisible);
+	bool GetVisibility(const UE::Geometry::FDynamicMesh3& Mesh, int TID);
+	void SetInternal(UE::Geometry::FDynamicMesh3& Mesh, int TID, bool bIsInternal);
+	bool GetInternal(const UE::Geometry::FDynamicMesh3& Mesh, int TID);
+	void SetUV(UE::Geometry::FDynamicMesh3& Mesh, int VID, FVector2f UV, int UVLayer);
+	void GetUV(const UE::Geometry::FDynamicMesh3& Mesh, int VID, FVector2f& UV, int UVLayer);
+	void SetTangent(UE::Geometry::FDynamicMesh3& Mesh, int VID, FVector3f Normal, FVector3f TangentU, FVector3f TangentV);
+	void GetTangent(const UE::Geometry::FDynamicMesh3& Mesh, int VID, FVector3f& U, FVector3f& V);
 	/// Initialize UV overlays based on the custom AugmentedDynamicMesh per-vertex UV attributes.  Optionally use FirstUVLayer to skip layers
-	void PLANARCUT_API InitializeOverlayToPerVertexUVs(UE::Geometry::FDynamicMesh3& Mesh, int32 NumUVLayers, int32 FirstUVLayer = 0);
-	void PLANARCUT_API InitializeOverlayToPerVertexTangents(UE::Geometry::FDynamicMesh3& Mesh);
-	void PLANARCUT_API ComputeTangents(UE::Geometry::FDynamicMesh3& Mesh, bool bOnlyInternalSurfaces,
+	void InitializeOverlayToPerVertexUVs(UE::Geometry::FDynamicMesh3& Mesh, int32 NumUVLayers, int32 FirstUVLayer = 0);
+	void InitializeOverlayToPerVertexTangents(UE::Geometry::FDynamicMesh3& Mesh);
+	void ComputeTangents(UE::Geometry::FDynamicMesh3& Mesh, bool bOnlyInternalSurfaces,
 		bool bRecomputeNormals = true, bool bMakeSharpEdges = false, float SharpAngleDegrees = 60);
-	void PLANARCUT_API AddCollisionSamplesPerComponent(UE::Geometry::FDynamicMesh3& Mesh, double Spacing);
-	void PLANARCUT_API SplitOverlayAttributesToPerVertex(UE::Geometry::FDynamicMesh3& Mesh, bool bSplitUVs = true, bool bSplitTangents = true);
+	void AddCollisionSamplesPerComponent(UE::Geometry::FDynamicMesh3& Mesh, double Spacing);
+	void SplitOverlayAttributesToPerVertex(UE::Geometry::FDynamicMesh3& Mesh, bool bSplitUVs = true, bool bSplitTangents = true);
 }
 
 
 /**
  * Dynamic mesh representation of cutting cells, to be used to fracture a mesh
  */
-struct PLANARCUT_API FCellMeshes
+struct FCellMeshes
 {
 	struct FCellInfo
 	{
@@ -150,7 +156,7 @@ private:
 
 // Holds Geometry from an FGeometryCollection in an FDynamicMesh3 representation, and convert both directions
 // Also supports cutting geometry with FCellMeshes
-struct PLANARCUT_API FDynamicMeshCollection
+struct FDynamicMeshCollection
 {
 	struct FMeshData
 	{
@@ -198,6 +204,13 @@ struct PLANARCUT_API FDynamicMeshCollection
 	TIndirectArray<FMeshData> Meshes;
 	UE::Geometry::FAxisAlignedBox3d Bounds;
 	
+	// Settings to control the geometry import
+	
+	// If true, triangles where the Visible property is false will not be added to the MeshData
+	bool bSkipInvisible = false;
+	// If false, Transforms passed to Init are interpreted as relative to the parent bone transform. If true, Transforms are all in the same 'global' / component-relative space
+	bool bComponentSpaceTransforms = false;
+	
 	FDynamicMeshCollection() {}
 
 	FDynamicMeshCollection(const FGeometryCollection* Collection, const TArrayView<const int32>& TransformIndices, FTransform TransformCollection, bool bSaveIsolatedVertices = false)
@@ -220,8 +233,17 @@ struct PLANARCUT_API FDynamicMeshCollection
 		Init(Collection, Collection->Transform, TransformIndices, TransformCollection, bSaveIsolatedVertices);
 	}
 
-	void Init(const FGeometryCollection* Collection, const TManagedArray<FTransform>& Transforms, const TArrayView<const int32>& TransformIndices, FTransform TransformCollection, bool bSaveIsolatedVertices = false);
-	void Init(const FGeometryCollection* Collection, const TManagedArray<FTransform3f>& Transforms, const TArrayView<const int32>& TransformIndices, FTransform TransformCollection, bool bSaveIsolatedVertices = false);
+	void Init(const FGeometryCollection* Collection, const TManagedArray<FTransform>& Transforms, const TArrayView<const int32>& TransformIndices, FTransform TransformCollection, bool bSaveIsolatedVertices = false)
+	{
+		Init(Collection, TArrayView<const FTransform>(Transforms.GetConstArray()), TransformIndices, TransformCollection, bSaveIsolatedVertices);
+	}
+	void Init(const FGeometryCollection* Collection, const TManagedArray<FTransform3f>& Transforms, const TArrayView<const int32>& TransformIndices, FTransform TransformCollection, bool bSaveIsolatedVertices = false)
+	{
+		Init(Collection, TArrayView<const FTransform3f>(Transforms.GetConstArray()), TransformIndices, TransformCollection, bSaveIsolatedVertices);
+	}
+
+	void Init(const FGeometryCollection* Collection, TArrayView<const FTransform> Transforms, const TArrayView<const int32>& TransformIndices, FTransform TransformCollection, bool bSaveIsolatedVertices = false);
+	void Init(const FGeometryCollection* Collection, TArrayView<const FTransform3f> Transforms, const TArrayView<const int32>& TransformIndices, FTransform TransformCollection, bool bSaveIsolatedVertices = false);
 
 	int32 CutWithMultiplePlanes(
 		const TArrayView<const FPlane>& Planes,
@@ -285,7 +307,7 @@ struct PLANARCUT_API FDynamicMeshCollection
 private:
 
 	template<typename TransformType>
-	void InitTemplate(const FGeometryCollection* Collection, const TManagedArray<TransformType>& Transforms, const TArrayView<const int32>& TransformIndices, FTransform TransformCollection, bool bSaveIsolatedVertices);
+	void InitTemplate(const FGeometryCollection* Collection, TArrayView<const TransformType> Transforms, const TArrayView<const int32>& TransformIndices, FTransform TransformCollection, bool bSaveIsolatedVertices);
 
 
 	void SetGeometryVisibility(FGeometryCollection* Collection, const TArray<int32>& GeometryIndices, bool bVisible);
diff --git a/Engine/Plugins/Experimental/PlanarCutPlugin/Source/PlanarCut/Private/PlanarCut.cpp b/Engine/Plugins/Experimental/PlanarCutPlugin/Source/PlanarCut/Private/PlanarCut.cpp
--- a/Engine/Plugins/Experimental/PlanarCutPlugin/Source/PlanarCut/Private/PlanarCut.cpp
+++ b/Engine/Plugins/Experimental/PlanarCutPlugin/Source/PlanarCut/Private/PlanarCut.cpp
@@ -2086,20 +2086,27 @@ int32 AddCollisionSampleVertices(double CollisionSampleSpacing, FGeometryCollect
 }
 
 template <typename TransformType>
-void ConvertToMeshDescriptionTemplate(
-	FMeshDescription& MeshOut,
+void ConvertToDynamicMeshTemplate(
+	FDynamicMesh3& CombinedMesh,
 	FTransform& TransformOut,
 	bool bCenterPivot,
-	FGeometryCollection& Collection,
-	const TManagedArray<TransformType>& BoneTransforms,
-	const TArrayView<const int32>& TransformIndices,
-	TFunction<int32(int32, bool)> RemapMaterialIDs
+	const FGeometryCollection& Collection,
+	TArrayView<const TransformType> BoneTransforms,
+	TArrayView<const int32> TransformIndices,
+	TFunction<int32(int32, bool)> RemapMaterialIDs,
+	bool bClearCustomAttributes = true,
+	bool bWeldEdges = true,
+	bool bComponentSpaceTransforms = false,
+	bool bAllowInvisible = true,
+	bool bSetPolygroupPerBone = false
 )
 {
 	FTransform CellsToWorld = FTransform::Identity;
 	TransformOut = FTransform::Identity;
 
 	FDynamicMeshCollection MeshCollection;
+	MeshCollection.bSkipInvisible = !bAllowInvisible;
+	MeshCollection.bComponentSpaceTransforms = bComponentSpaceTransforms && !BoneTransforms.IsEmpty();
 	if (BoneTransforms.Num())
 	{
 		MeshCollection.Init(&Collection, BoneTransforms, TransformIndices, CellsToWorld);
@@ -2108,15 +2115,22 @@ void ConvertToMeshDescriptionTemplate(
 	{
 		MeshCollection.Init(&Collection, Collection.Transform, TransformIndices, CellsToWorld);
 	}
-	
-	FDynamicMesh3 CombinedMesh;
+
 	SetGeometryCollectionAttributes(CombinedMesh, Collection.NumUVLayers());
 	CombinedMesh.Attributes()->EnableTangents();
+	if (bSetPolygroupPerBone)
+	{
+		CombinedMesh.EnableTriangleGroups();
+	}
 
 	int32 NumMeshes = MeshCollection.Meshes.Num();
 	for (int32 MeshIdx = 0; MeshIdx < NumMeshes; MeshIdx++)
 	{
 		FDynamicMesh3& Mesh = MeshCollection.Meshes[MeshIdx].AugMesh;
+		if (bSetPolygroupPerBone)
+		{
+			Mesh.EnableTriangleGroups();
+		}
 		const FTransform& FromCollection = MeshCollection.Meshes[MeshIdx].FromCollection;
 
 		FMeshNormals::InitializeOverlayToPerVertexNormals(Mesh.Attributes()->PrimaryNormals(), true);
@@ -2135,9 +2149,12 @@ void ConvertToMeshDescriptionTemplate(
 			}
 		}
 
-		FMergeCoincidentMeshEdges EdgeMerge(&Mesh);
-		EdgeMerge.Apply();
-		
+		if (bWeldEdges && Mesh.TriangleCount() > 0)
+		{
+			FMergeCoincidentMeshEdges EdgeMerge(&Mesh);
+			EdgeMerge.Apply();
+		}
+
 		if (MeshIdx > 0)
 		{
 			FDynamicMeshEditor MeshAppender(&CombinedMesh);
@@ -2158,6 +2175,26 @@ void ConvertToMeshDescriptionTemplate(
 		TransformOut = FTransform((FVector)-Translate);
 	}
 
+	if (bClearCustomAttributes)
+	{
+		ClearCustomGeometryCollectionAttributes(CombinedMesh);
+	}
+}
+
+template <typename TransformType>
+void ConvertToMeshDescriptionTemplate(
+	FMeshDescription& MeshOut,
+	FTransform& TransformOut,
+	bool bCenterPivot,
+	FGeometryCollection& Collection,
+	const TManagedArray<TransformType>& BoneTransforms,
+	const TArrayView<const int32>& TransformIndices,
+	TFunction<int32(int32, bool)> RemapMaterialIDs
+)
+{	
+	FDynamicMesh3 CombinedMesh;
+	ConvertToDynamicMeshTemplate<TransformType>(CombinedMesh, TransformOut, bCenterPivot, Collection, BoneTransforms.GetConstArray(), TransformIndices, RemapMaterialIDs, false);
+
 	FDynamicMeshToMeshDescription Converter;
 	Converter.Convert(&CombinedMesh, MeshOut, true);
 }
@@ -2188,5 +2225,20 @@ void ConvertToMeshDescription(
 	ConvertToMeshDescriptionTemplate(MeshOut, TransformOut, bCenterPivot, Collection, BoneTransforms, TransformIndices, RemapMaterialIDs);
 }
 
+void ConvertGeometryCollectionToDynamicMesh(
+	FDynamicMesh3& OutputMesh,
+	FTransform& TransformOut,
+	bool bCenterPivot,
+	const FGeometryCollection& Collection,
+	bool bWeldEdges,
+	TArrayView<const FTransform3f> BoneTransforms,
+	bool bUseRelativeTransforms,
+	TArrayView<const int32> TransformIndices,
+	TFunction<int32(int32, bool)> RemapMaterialIDs,
+	bool bAllowInvisible, bool bSetPolygroupPerBone
+)
+{
+	ConvertToDynamicMeshTemplate<FTransform3f>(OutputMesh, TransformOut, bCenterPivot, Collection, BoneTransforms, TransformIndices, RemapMaterialIDs, true, bWeldEdges, !bUseRelativeTransforms, bAllowInvisible, bSetPolygroupPerBone);
+}
 
 #undef LOCTEXT_NAMESPACE
diff --git a/Engine/Plugins/Experimental/PlanarCutPlugin/Source/PlanarCut/Public/PlanarCut.h b/Engine/Plugins/Experimental/PlanarCutPlugin/Source/PlanarCut/Public/PlanarCut.h
--- a/Engine/Plugins/Experimental/PlanarCutPlugin/Source/PlanarCut/Public/PlanarCut.h
+++ b/Engine/Plugins/Experimental/PlanarCutPlugin/Source/PlanarCut/Public/PlanarCut.h
@@ -557,6 +557,34 @@ void PLANARCUT_API ConvertToMeshDescription(
 	TFunction<int32(int32, bool)> RemapMaterialIDs = nullptr
 );
 
+/**
+ * Convert chosen Geometry groups inside a GeometryCollection to a single Dynamic Mesh.
+ *
+ * @param OutputMesh				Mesh to be filled with the geometry collection geometry
+ * @param TransformOut				Transform taking output mesh geometry to local space of geometry collection. This will be Identity unless bCenterPivot is true.
+ * @param bCenterPivot				Whether to center the geometry at the origin
+ * @param Collection				The collection to be converted
+ * @param bWeldEdges				Whether to weld edges on conversion
+ * @param BoneTransforms			Transforms for each bone in the geometry collection
+ * @param bUseRelativeTransforms	Whether the BoneTransforms are relative to each bone's parent, or in a shared space (e.g., Collection.Transform is relative, UGeometryCollectionComponent's component space transforms are not)
+ * @param TransformIndices			Which transform groups inside the collection to convert
+ * @param RemapMaterialIDs			Optional function to remap (MaterialID, bIsInternal) -> NewMaterialID
+ * @param bAllowInvisible			Allow mesh to include triangles that are not marked 'visible' (typically, pre-fracture geometry stored for potential history/LOD usage)
+ * @param bSetPolygroupPerBone		Enable PolyGroups on the OutputMesh and set one group per geometry
+ */
+void PLANARCUT_API ConvertGeometryCollectionToDynamicMesh(
+	UE::Geometry::FDynamicMesh3& OutputMesh,
+	FTransform& TransformOut,
+	bool bCenterPivot,
+	const FGeometryCollection& Collection,
+	bool bWeldEdges,
+	TArrayView<const FTransform3f> BoneTransforms,
+	bool bUseRelativeTransforms,
+	TArrayView<const int32> TransformIndices,
+	TFunction<int32(int32, bool)> RemapMaterialIDs = nullptr,
+	bool bAllowInvisible = false,
+	bool bSetPolygroupPerBone = true
+);
 
 #if UE_ENABLE_INCLUDE_ORDER_DEPRECATED_IN_5_2
 #include "CoreMinimal.h"
diff --git a/Engine/Source/Runtime/Experimental/Chaos/Private/GeometryCollection/GeometryCollectionAlgo.cpp b/Engine/Source/Runtime/Experimental/Chaos/Private/GeometryCollection/GeometryCollectionAlgo.cpp
--- a/Engine/Source/Runtime/Experimental/Chaos/Private/GeometryCollection/GeometryCollectionAlgo.cpp
+++ b/Engine/Source/Runtime/Experimental/Chaos/Private/GeometryCollection/GeometryCollectionAlgo.cpp
@@ -560,11 +560,27 @@ namespace GeometryCollectionAlgo
 
 		if (RelativeTransforms.IsValidIndex(Index))
 		{
-			while (Index != FGeometryCollection::Invalid)
+			do
 			{
 				Transform = Transform * RelativeTransforms[Index];
 				Index = Parents[Index];
-			}
+			} while (Index != FGeometryCollection::Invalid);
+		}
+		return Transform;
+	}
+
+	template<typename TransformType>
+	TransformType GlobalMatrixTemplate(TArrayView<const TransformType> RelativeTransforms, TArrayView<const int32> Parents, int32 Index)
+	{
+		TransformType Transform = TransformType::Identity;
+
+		if (RelativeTransforms.IsValidIndex(Index))
+		{
+			do
+			{
+				Transform = Transform * RelativeTransforms[Index];
+				Index = Parents[Index];
+			} while (Index != FGeometryCollection::Invalid);
 		}
 		return Transform;
 	}
@@ -585,6 +601,15 @@ namespace GeometryCollectionAlgo
 		return FTransform(GlobalMatrixTemplate<FTransform3f>(RelativeTransforms, Parents, Index));
 	}
 
+	FTransform GlobalMatrix(TArrayView<const FTransform> RelativeTransforms, TArrayView<const int32> Parents, int32 Index)
+	{
+		return GlobalMatrixTemplate<FTransform>(RelativeTransforms, Parents, Index);
+	}
+	FTransform GlobalMatrix(TArrayView<const FTransform3f> RelativeTransforms, TArrayView<const int32> Parents, int32 Index)
+	{
+		return FTransform(GlobalMatrixTemplate<FTransform3f>(RelativeTransforms, Parents, Index));
+	}
+
 	namespace Private
 	{
 		void GlobalMatrices(const FGeometryDynamicCollection& DynamicCollection, const TArray<int32>& Indices, TArray<FTransform>& OutGlobalTransforms)
diff --git a/Engine/Source/Runtime/Experimental/Chaos/Public/GeometryCollection/GeometryCollectionAlgo.h b/Engine/Source/Runtime/Experimental/Chaos/Public/GeometryCollection/GeometryCollectionAlgo.h
--- a/Engine/Source/Runtime/Experimental/Chaos/Public/GeometryCollection/GeometryCollectionAlgo.h
+++ b/Engine/Source/Runtime/Experimental/Chaos/Public/GeometryCollection/GeometryCollectionAlgo.h
@@ -135,6 +135,8 @@ namespace GeometryCollectionAlgo
 	*/
 	FTransform CHAOS_API GlobalMatrix(const TManagedArray<FTransform>& RelativeTransforms, const TManagedArray<int32>& Parents, int32 Index);
 	FTransform CHAOS_API GlobalMatrix(const TManagedArray<FTransform3f>& RelativeTransforms, const TManagedArray<int32>& Parents, int32 Index);
+	FTransform CHAOS_API GlobalMatrix(TArrayView<const FTransform> RelativeTransforms, TArrayView<const int32> Parents, int32 Index);
+	FTransform CHAOS_API GlobalMatrix(TArrayView<const FTransform3f> RelativeTransforms, TArrayView<const int32> Parents, int32 Index);
 	FTransform3f CHAOS_API GlobalMatrix3f(const TManagedArray<FTransform3f>& RelativeTransforms, const TManagedArray<int32>& Parents, int32 Index);
 
 
