add option to not split islands when fracturing meshes

diff --git a/Engine/Plugins/Experimental/ChaosEditor/Source/FractureEditor/Private/FractureToolBrick.cpp b/Engine/Plugins/Experimental/ChaosEditor/Source/FractureEditor/Private/FractureToolBrick.cpp
--- a/Engine/Plugins/Experimental/ChaosEditor/Source/FractureEditor/Private/FractureToolBrick.cpp
+++ b/Engine/Plugins/Experimental/ChaosEditor/Source/FractureEditor/Private/FractureToolBrick.cpp
@@ -394,6 +394,7 @@ public:
 	float Grout = 0;
 	int Seed;
 	FTransform Transform;
+	bool bSplitIslands = true;
 
 	// TGenericDataOperator interface:
 	virtual void CalculateResult(FProgressCancel* Progress) override
@@ -403,7 +404,7 @@ public:
 			return;
 		}
 
-		ResultGeometryIndex = CutMultipleWithPlanarCells(Cells, *CollectionCopy, Selection, Grout, PointSpacing, Seed, Transform, true, true, Progress, CellsOrigin);
+		ResultGeometryIndex = CutMultipleWithPlanarCells(Cells, *CollectionCopy, Selection, Grout, PointSpacing, Seed, Transform, true, true, Progress, CellsOrigin, bSplitIslands);
 
 		SetResult(MoveTemp(CollectionCopy));
 	}
@@ -448,6 +449,7 @@ int32 UFractureToolBrick::ExecuteFracture(const FFractureToolContext& FractureCo
 		BrickOp->Selection = FractureContext.GetSelection();
 		BrickOp->Grout = 0; // CutterSettings->Grout; // Note: Grout is currently baked directly into the brick cells above
 		BrickOp->PointSpacing = CollisionSettings->GetPointSpacing();
+		BrickOp->bSplitIslands = CutterSettings->bSplitIslands;
 		const bool bBricksAreTouching = CutterSettings->Grout <= UE_KINDA_SMALL_NUMBER;
 		BrickOp->Cells = FPlanarCells(BricksToCut, bBricksAreTouching);
 		if (CutterSettings->Amplitude > 0.0f)
diff --git a/Engine/Plugins/Experimental/ChaosEditor/Source/FractureEditor/Private/FractureToolCutter.cpp b/Engine/Plugins/Experimental/ChaosEditor/Source/FractureEditor/Private/FractureToolCutter.cpp
--- a/Engine/Plugins/Experimental/ChaosEditor/Source/FractureEditor/Private/FractureToolCutter.cpp
+++ b/Engine/Plugins/Experimental/ChaosEditor/Source/FractureEditor/Private/FractureToolCutter.cpp
@@ -626,6 +626,7 @@ public:
 	float Grout;
 	int Seed;
 	FTransform Transform;
+	bool bSplitIslands = true;
 
 	// TGenericDataOperator interface:
 	virtual void CalculateResult(FProgressCancel* Progress) override
@@ -657,7 +658,7 @@ public:
 		FProgressCancel::FProgressScope FractureMeshProgress =
 			FProgressCancel::CreateScopeTo(Progress, 1, LOCTEXT("FractureMeshMessage", "Fracturing Mesh"));
 
-		ResultGeometryIndex = CutMultipleWithPlanarCells(VoronoiPlanarCells, *CollectionCopy, Selection, Grout, PointSpacing, Seed, Transform, true, true, Progress, Origin);
+		ResultGeometryIndex = CutMultipleWithPlanarCells(VoronoiPlanarCells, *CollectionCopy, Selection, Grout, PointSpacing, Seed, Transform, true, true, Progress, Origin, bSplitIslands);
 		
 		SetResult(MoveTemp(CollectionCopy));
 	}
@@ -678,6 +679,7 @@ int32 UFractureToolVoronoiCutterBase::ExecuteFracture(const FFractureToolContext
 		VoronoiOp->Grout = CutterSettings->Grout;
 		VoronoiOp->PointSpacing = CollisionSettings->GetPointSpacing();
 		VoronoiOp->Sites = Sites;
+		VoronoiOp->bSplitIslands = CutterSettings->bSplitIslands;
 		if (CutterSettings->Amplitude > 0.0f)
 		{
 			FNoiseSettings Settings;
diff --git a/Engine/Plugins/Experimental/ChaosEditor/Source/FractureEditor/Private/FractureToolCutter.h b/Engine/Plugins/Experimental/ChaosEditor/Source/FractureEditor/Private/FractureToolCutter.h
--- a/Engine/Plugins/Experimental/ChaosEditor/Source/FractureEditor/Private/FractureToolCutter.h
+++ b/Engine/Plugins/Experimental/ChaosEditor/Source/FractureEditor/Private/FractureToolCutter.h
@@ -82,6 +82,10 @@ public:
 	UPROPERTY()
 	bool bGroupFractureToggleEnabled = true;
 
+	// Whether to split the fractured mesh pieces based on geometric connectivity after fracturing
+	UPROPERTY(EditAnywhere, Category = CommonFracture)
+	bool bSplitIslands = true;
+
 	/** Amount of space to leave between cut pieces */
 	UPROPERTY(EditAnywhere, Category = CommonFracture, meta = (UIMin = "0.0", ClampMin = "0.0", EditCondition = "bGroutSettingEnabled", HideEditConditionToggle, EditConditionHides))
 	float Grout = 0.0f;
diff --git a/Engine/Plugins/Experimental/ChaosEditor/Source/FractureEditor/Private/FractureToolMeshCut.cpp b/Engine/Plugins/Experimental/ChaosEditor/Source/FractureEditor/Private/FractureToolMeshCut.cpp
--- a/Engine/Plugins/Experimental/ChaosEditor/Source/FractureEditor/Private/FractureToolMeshCut.cpp
+++ b/Engine/Plugins/Experimental/ChaosEditor/Source/FractureEditor/Private/FractureToolMeshCut.cpp
@@ -221,6 +221,7 @@ public:
 	int Seed;
 	FTransform Transform;
 	UE::Geometry::FDynamicMesh3 CuttingMesh;
+	bool bSplitIslands = true;
 
 	// TGenericDataOperator interface:
 	virtual void CalculateResult(FProgressCancel* Progress) override
@@ -242,7 +243,8 @@ public:
 			}
 
 			FProgressCancel::FProgressScope ProgressScope(Progress, ProgressFrac);
-			int32 Index = CutWithMesh(CuttingMesh, ScatterTransform, InternalSurfaceMaterials, *CollectionCopy, BonesToCut, PointSpacing, Transform);
+			constexpr bool bSetDefaultInternalMaterialsFromCollection = true;
+			int32 Index = CutWithMesh(CuttingMesh, ScatterTransform, InternalSurfaceMaterials, *CollectionCopy, BonesToCut, PointSpacing, Transform, bSetDefaultInternalMaterialsFromCollection, Progress, bSplitIslands);
 			if (Progress && Progress->Cancelled())
 			{
 				return;
@@ -319,6 +321,7 @@ int32 UFractureToolMeshCut::ExecuteFracture(const FFractureToolContext& Fracture
 		MeshCutOp->PointSpacing = CollisionSettings->GetPointSpacing();
 		MeshCutOp->Seed = FractureContext.GetSeed();
 		MeshCutOp->Transform = FractureContext.GetTransform();
+		MeshCutOp->bSplitIslands = CutterSettings->bSplitIslands;
 		if (LocalCutSettings->CutDistribution == EMeshCutDistribution::SingleCut)
 		{
 			MeshCutOp->MeshTransforms.Add(LocalCutSettings->CuttingActor->GetTransform());
diff --git a/Engine/Plugins/Experimental/ChaosEditor/Source/FractureEditor/Private/FractureToolPlaneCut.cpp b/Engine/Plugins/Experimental/ChaosEditor/Source/FractureEditor/Private/FractureToolPlaneCut.cpp
--- a/Engine/Plugins/Experimental/ChaosEditor/Source/FractureEditor/Private/FractureToolPlaneCut.cpp
+++ b/Engine/Plugins/Experimental/ChaosEditor/Source/FractureEditor/Private/FractureToolPlaneCut.cpp
@@ -379,11 +379,12 @@ public:
 	int Seed;
 	FTransform Transform;
 	UE::Geometry::FDynamicMesh3 CuttingMesh;
+	bool bSplitIslands = true;
 
 	// TGenericDataOperator interface:
 	virtual void CalculateResult(FProgressCancel* Progress) override
 	{
-		ResultGeometryIndex = CutMultipleWithMultiplePlanes(CuttingPlanes, InternalSurfaceMaterials, *CollectionCopy, Selection, Grout, PointSpacing, Seed, Transform, true, Progress);
+		ResultGeometryIndex = CutMultipleWithMultiplePlanes(CuttingPlanes, InternalSurfaceMaterials, *CollectionCopy, Selection, Grout, PointSpacing, Seed, Transform, true, Progress, bSplitIslands);
 
 		if (Progress && Progress->Cancelled())
 		{
@@ -404,6 +405,7 @@ int32 UFractureToolPlaneCut::ExecuteFracture(const FFractureToolContext& Fractur
 		PlaneCutOp->PointSpacing = CollisionSettings->GetPointSpacing();
 		PlaneCutOp->Seed = FractureContext.GetSeed();
 		PlaneCutOp->Transform = FractureContext.GetTransform();
+		PlaneCutOp->bSplitIslands = CutterSettings->bSplitIslands;
 
 		if (GizmoSettings->IsGizmoEnabled())
 		{
diff --git a/Engine/Plugins/Experimental/ChaosEditor/Source/FractureEditor/Private/FractureToolSlice.cpp b/Engine/Plugins/Experimental/ChaosEditor/Source/FractureEditor/Private/FractureToolSlice.cpp
--- a/Engine/Plugins/Experimental/ChaosEditor/Source/FractureEditor/Private/FractureToolSlice.cpp
+++ b/Engine/Plugins/Experimental/ChaosEditor/Source/FractureEditor/Private/FractureToolSlice.cpp
@@ -187,7 +187,7 @@ int32 UFractureToolSlice::ExecuteFracture(const FFractureToolContext& FractureCo
 		// Proximity is invalidated.
 		ClearProximity(FractureContext.GetGeometryCollection().Get());
 
-		return CutMultipleWithMultiplePlanes(CuttingPlanes, InternalSurfaceMaterials, *FractureContext.GetGeometryCollection(), FractureContext.GetSelection(), CutterSettings->Grout, CollisionSettings->GetPointSpacing(), FractureContext.GetSeed(), FractureContext.GetTransform());
+		return CutMultipleWithMultiplePlanes(CuttingPlanes, InternalSurfaceMaterials, *FractureContext.GetGeometryCollection(), FractureContext.GetSelection(), CutterSettings->Grout, CollisionSettings->GetPointSpacing(), FractureContext.GetSeed(), FractureContext.GetTransform(), true, nullptr, CutterSettings->bSplitIslands);
 	}
 
 	return INDEX_NONE;
diff --git a/Engine/Plugins/Experimental/PlanarCutPlugin/Source/PlanarCut/Private/GeometryMeshConversion.cpp b/Engine/Plugins/Experimental/PlanarCutPlugin/Source/PlanarCut/Private/GeometryMeshConversion.cpp
--- a/Engine/Plugins/Experimental/PlanarCutPlugin/Source/PlanarCut/Private/GeometryMeshConversion.cpp
+++ b/Engine/Plugins/Experimental/PlanarCutPlugin/Source/PlanarCut/Private/GeometryMeshConversion.cpp
@@ -2685,6 +2685,7 @@ int32 FDynamicMeshCollection::CutWithMultiplePlanes(
 	const TArrayView<const FPlane>& Planes,
 	double Grout,
 	double CollisionSampleSpacing,
+	bool bSplitIslands,
 	int32 RandomSeed,
 	FGeometryCollection* Collection,
 	FInternalSurfaceMaterials& InternalSurfaceMaterials,
@@ -2746,7 +2747,7 @@ int32 FDynamicMeshCollection::CutWithMultiplePlanes(
 		TArray<TPair<int32, int32>> CellConnectivity;
 		CellConnectivity.Add(TPair<int32, int32>(0, -1));
 
-		return CutWithCellMeshes(InternalSurfaceMaterials, CellConnectivity, GroutCells, Collection, bSetDefaultInternalMaterialsFromCollection, CollisionSampleSpacing);
+		return CutWithCellMeshes(InternalSurfaceMaterials, CellConnectivity, GroutCells, bSplitIslands, Collection, bSetDefaultInternalMaterialsFromCollection, CollisionSampleSpacing);
 	}
 
 	TArray<TUniquePtr<FMeshData>> ToCut;
@@ -2836,7 +2837,7 @@ int32 FDynamicMeshCollection::CutWithMultiplePlanes(
 				TArray<FDynamicMesh3> SplitMeshes;
 				for (int UnsplitIdx = 0; UnsplitIdx < 2; UnsplitIdx++)
 				{
-					if (SplitIslands(ToCut[ResultIndices[UnsplitIdx]]->AugMesh, SplitMeshes))
+					if (bSplitIslands && SplitIslands(ToCut[ResultIndices[UnsplitIdx]]->AugMesh, SplitMeshes))
 					{
 						ToCut[ResultIndices[UnsplitIdx]]->SetMesh(SplitMeshes[0]);
 						for (int32 Idx = 1; Idx < SplitMeshes.Num(); Idx++)
@@ -2984,7 +2985,7 @@ int32 FDynamicMeshCollection::SplitAllIslands(FGeometryCollection* Collection, d
 
 
 
-int32 FDynamicMeshCollection::CutWithCellMeshes(const FInternalSurfaceMaterials& InternalSurfaceMaterials, const TArray<TPair<int32, int32>>& CellConnectivity, FCellMeshes& CellMeshes, FGeometryCollection* Collection, bool bSetDefaultInternalMaterialsFromCollection, double CollisionSampleSpacing)
+int32 FDynamicMeshCollection::CutWithCellMeshes(const FInternalSurfaceMaterials& InternalSurfaceMaterials, const TArray<TPair<int32, int32>>& CellConnectivity, FCellMeshes& CellMeshes, bool bSplitIslands, FGeometryCollection* Collection, bool bSetDefaultInternalMaterialsFromCollection, double CollisionSampleSpacing)
 {
 	// TODO: should we do these cuts in parallel, and the appends sequentially below?
 	int32 FirstIdx = -1;
@@ -3065,7 +3066,7 @@ int32 FDynamicMeshCollection::CutWithCellMeshes(const FInternalSurfaceMaterials&
 					}
 					int32 CreatedGeometryIdx = -1;
 					TArray<FDynamicMesh3> Islands;
-					if (SplitIslands(AugBoolResult, Islands))
+					if (bSplitIslands && SplitIslands(AugBoolResult, Islands))
 					{
 						for (int32 i = 0; i < Islands.Num(); i++)
 						{
@@ -3183,22 +3184,19 @@ bool FDynamicMeshCollection::IsNeighboring(
 }
 
 // Split mesh into connected components, including implicit connections by co-located vertices
-bool FDynamicMeshCollection::SplitIslands(FDynamicMesh3& Source, TArray<FDynamicMesh3>& SeparatedMeshes)
+bool FDynamicMeshCollection::SplitIslands(FDynamicMesh3& Source, TArray<FDynamicMesh3>& SeparatedMeshes, double SnapDistance)
 {
-	double SnapDistance = 1e-03;
 	TPointHashGrid3d<int> VertHash(SnapDistance * 10, -1);
 	FDisjointSet VertComponents(Source.MaxVertexID());
 	// Add Source vertices to hash & disjoint sets
-	TArray<int> Neighbors;
 	for (int VID : Source.VertexIndicesItr())
 	{
 		FVector3d Pt = Source.GetVertex(VID);
-		Neighbors.Reset();
-		VertHash.FindPointsInBall(Pt, SnapDistance, [&Source, Pt](int OtherVID) {return DistanceSquared(Pt, Source.GetVertex(OtherVID)); }, Neighbors);
-		for (int NbrVID : Neighbors)
-		{
-			VertComponents.UnionSequential(VID, NbrVID);
-		}
+		VertHash.EnumeratePointsInBall(Pt, SnapDistance, [&Source, Pt](int OtherVID) {return DistanceSquared(Pt, Source.GetVertex(OtherVID)); }, [&](int32 NbrVID, double DSq) -> bool
+			{
+				VertComponents.UnionSequential(VID, NbrVID);
+				return true;
+			});
 		VertHash.InsertPointUnsafe(VID, Pt);
 	}
 	for (FIndex3i Tri : Source.TrianglesItr())
diff --git a/Engine/Plugins/Experimental/PlanarCutPlugin/Source/PlanarCut/Private/GeometryMeshConversion.h b/Engine/Plugins/Experimental/PlanarCutPlugin/Source/PlanarCut/Private/GeometryMeshConversion.h
--- a/Engine/Plugins/Experimental/PlanarCutPlugin/Source/PlanarCut/Private/GeometryMeshConversion.h
+++ b/Engine/Plugins/Experimental/PlanarCutPlugin/Source/PlanarCut/Private/GeometryMeshConversion.h
@@ -249,6 +249,7 @@ struct FDynamicMeshCollection
 		const TArrayView<const FPlane>& Planes,
 		double Grout,
 		double CollisionSampleSpacing,
+		bool bSplitIslands,
 		int32 RandomSeed,
 		FGeometryCollection* Collection,
 		FInternalSurfaceMaterials& InternalSurfaceMaterials,
@@ -267,7 +268,7 @@ struct FDynamicMeshCollection
 	 * @param CollisionSampleSpacing If > 0, new geometry will have collision samples added (vertices not on any triangles) to fill any gaps greater than the this size
 	 * @return Index of the first created geometry
 	 */
-	int32 CutWithCellMeshes(const FInternalSurfaceMaterials& InternalSurfaceMaterials, const TArray<TPair<int32, int32>>& CellConnectivity, FCellMeshes& CellMeshes, FGeometryCollection* Collection, bool bSetDefaultInternalMaterialsFromCollection, double CollisionSampleSpacing);
+	int32 CutWithCellMeshes(const FInternalSurfaceMaterials& InternalSurfaceMaterials, const TArray<TPair<int32, int32>>& CellConnectivity, FCellMeshes& CellMeshes, bool bSplitIslands, FGeometryCollection* Collection, bool bSetDefaultInternalMaterialsFromCollection, double CollisionSampleSpacing);
 
 
 	/**
@@ -289,7 +290,7 @@ struct FDynamicMeshCollection
 	}
 
 	// Split mesh into connected components, including implicit connections by co-located vertices
-	bool SplitIslands(UE::Geometry::FDynamicMesh3& Source, TArray<UE::Geometry::FDynamicMesh3>& SeparatedMeshes);
+	bool SplitIslands(UE::Geometry::FDynamicMesh3& Source, TArray<UE::Geometry::FDynamicMesh3>& SeparatedMeshes, double ToleranceDistance = 1e-3);
 
 	FString GetBoneName(FGeometryCollection& Output, int TransformParent, int SubPartIndex)
 	{
diff --git a/Engine/Plugins/Experimental/PlanarCutPlugin/Source/PlanarCut/Private/PlanarCut.cpp b/Engine/Plugins/Experimental/PlanarCutPlugin/Source/PlanarCut/Private/PlanarCut.cpp
--- a/Engine/Plugins/Experimental/PlanarCutPlugin/Source/PlanarCut/Private/PlanarCut.cpp
+++ b/Engine/Plugins/Experimental/PlanarCutPlugin/Source/PlanarCut/Private/PlanarCut.cpp
@@ -828,7 +828,8 @@ int32 CutMultipleWithMultiplePlanes(
 	int32 RandomSeed,
 	const TOptional<FTransform>& TransformCollection,
 	bool bSetDefaultInternalMaterialsFromCollection,
-	FProgressCancel* Progress
+	FProgressCancel* Progress,
+	bool bSplitIslands
 )
 {
 	FProgressCancel::FProgressScope PrepareScope = FProgressCancel::CreateScopeTo(Progress, .1, LOCTEXT("CutWithMultiplePlanesInit", "Preparing to cut with planes"));
@@ -869,7 +870,7 @@ int32 CutMultipleWithMultiplePlanes(
 	FProgressCancel::FProgressScope CutScope = FProgressCancel::CreateScopeTo(Progress, .99, LOCTEXT("CutWithMultiplePlanesBody", "Cutting with planes"));
 
 	int32 NewGeomStartIdx = -1;
-	NewGeomStartIdx = MeshCollection.CutWithMultiplePlanes(CenteredPlanes, Grout, CollisionSampleSpacing, RandomSeed, &Collection, InternalSurfaceMaterials, bSetDefaultInternalMaterialsFromCollection, Progress);
+	NewGeomStartIdx = MeshCollection.CutWithMultiplePlanes(CenteredPlanes, Grout, CollisionSampleSpacing, bSplitIslands, RandomSeed, &Collection, InternalSurfaceMaterials, bSetDefaultInternalMaterialsFromCollection, Progress);
 
 	CutScope.Done();
 	if (Progress && Progress->Cancelled())
@@ -942,7 +943,8 @@ int32 CutMultipleWithPlanarCells(
 	bool bIncludeOutsideCellInOutput,
 	bool bSetDefaultInternalMaterialsFromCollection,
 	FProgressCancel* Progress,
-	FVector CellsOrigin
+	FVector CellsOrigin,
+	bool bSplitIslands
 )
 {
 	FProgressCancel::FProgressScope CreateMeshCollectionScope = FProgressCancel::CreateScopeTo(Progress, .1);
@@ -976,7 +978,7 @@ int32 CutMultipleWithPlanarCells(
 
 	FProgressCancel::FProgressScope CutScope = FProgressCancel::CreateScopeTo(Progress, .99);
 	int32 NewGeomStartIdx = -1;
-	NewGeomStartIdx = MeshCollection.CutWithCellMeshes(Cells.InternalSurfaceMaterials, Cells.PlaneCells, CellMeshes, &Source, bSetDefaultInternalMaterialsFromCollection, CollisionSampleSpacing);
+	NewGeomStartIdx = MeshCollection.CutWithCellMeshes(Cells.InternalSurfaceMaterials, Cells.PlaneCells, CellMeshes, bSplitIslands, &Source, bSetDefaultInternalMaterialsFromCollection, CollisionSampleSpacing);
 	CutScope.Done();
 
 	if (Progress && Progress->Cancelled())
@@ -1139,7 +1141,8 @@ int32 CutWithMesh(
 	double CollisionSampleSpacing,
 	const TOptional<FTransform>& TransformCollection,
 	bool bSetDefaultInternalMaterialsFromCollection,
-	FProgressCancel* Progress
+	FProgressCancel* Progress,
+	bool bSplitIslands
 )
 {
 	FProgressCancel::FProgressScope PrepareScope = FProgressCancel::CreateScopeTo(Progress, .1);
@@ -1179,7 +1182,7 @@ int32 CutWithMesh(
 	}
 	FProgressCancel::FProgressScope CutScope = FProgressCancel::CreateScopeTo(Progress, .99);
 
-	NewGeomStartIdx = MeshCollection.CutWithCellMeshes(InternalSurfaceMaterials, CellConnectivity, CellMeshes, &Collection, bSetDefaultInternalMaterialsFromCollection, CollisionSampleSpacing);
+	NewGeomStartIdx = MeshCollection.CutWithCellMeshes(InternalSurfaceMaterials, CellConnectivity, CellMeshes, bSplitIslands, &Collection, bSetDefaultInternalMaterialsFromCollection, CollisionSampleSpacing);
 
 	CutScope.Done();
 	if (Progress && Progress->Cancelled())
diff --git a/Engine/Plugins/Experimental/PlanarCutPlugin/Source/PlanarCut/Public/PlanarCut.h b/Engine/Plugins/Experimental/PlanarCutPlugin/Source/PlanarCut/Public/PlanarCut.h
--- a/Engine/Plugins/Experimental/PlanarCutPlugin/Source/PlanarCut/Public/PlanarCut.h
+++ b/Engine/Plugins/Experimental/PlanarCutPlugin/Source/PlanarCut/Public/PlanarCut.h
@@ -202,6 +202,7 @@ struct PLANARCUT_API FPlanarCells
  * @param bIncludeOutsideCellInOutput	If true, geometry that was not inside any of the cells (e.g. was outside of the bounds of all cutting geometry) will still be included in the output; if false, it will be discarded.
  * @param Progress						Optionally tracks progress and supports early-cancel
  * @param CellsOrigin					Optionally provide a local origin of the cutting Cells
+ * @param bSplitIslands					Whether to detect connected components and split them apart after the cut
  * @return	index of first new geometry in the Output GeometryCollection, or -1 if no geometry was added
  */
 int32 PLANARCUT_API CutWithPlanarCells(
@@ -215,7 +216,8 @@ int32 PLANARCUT_API CutWithPlanarCells(
 	bool bIncludeOutsideCellInOutput = true,
 	bool bSetDefaultInternalMaterialsFromCollection = true,
 	FProgressCancel* Progress = nullptr,
-	FVector CellsOrigin = FVector::ZeroVector
+	FVector CellsOrigin = FVector::ZeroVector,
+	bool bSplitIslands = true
 );
 
 /**
@@ -256,6 +258,7 @@ void PLANARCUT_API CreateCuttingSurfacePreview(
  * @param bIncludeOutsideCellInOutput	If true, geometry that was not inside any of the cells (e.g. was outside of the bounds of all cutting geometry) will still be included in the output; if false, it will be discarded.
  * @param Progress						Optionally tracks progress and supports early-cancel
  * @param CellsOrigin					Optionally provide a local origin of the cutting Cells
+ * @param bSplitIslands					Whether to detect connected components and split them apart after the cut
  * @return	index of first new geometry in the Output GeometryCollection, or -1 if no geometry was added
  */
 int32 PLANARCUT_API CutMultipleWithPlanarCells(
@@ -269,7 +272,8 @@ int32 PLANARCUT_API CutMultipleWithPlanarCells(
 	bool bIncludeOutsideCellInOutput = true,
 	bool bSetDefaultInternalMaterialsFromCollection = true,
 	FProgressCancel* Progress = nullptr,
-	FVector CellsOrigin = FVector::ZeroVector
+	FVector CellsOrigin = FVector::ZeroVector,
+	bool bSplitIslands = true
 );
 
 /**
@@ -300,6 +304,7 @@ int32 PLANARCUT_API SplitIslands(
  * @param RandomSeed				Seed to be used for random noise displacement
  * @param TransformCollection		Optional transform of the whole geometry collection; if unset, defaults to Identity
  * @param Progress					Optionally tracks progress and supports early-cancel
+ * @param bSplitIslands					Whether to detect connected components and split them apart after the cut
  * @return	index of first new geometry in the Output GeometryCollection, or -1 if no geometry was added
  */
 int32 PLANARCUT_API CutMultipleWithMultiplePlanes(
@@ -312,7 +317,8 @@ int32 PLANARCUT_API CutMultipleWithMultiplePlanes(
 	int32 RandomSeed,
 	const TOptional<FTransform>& TransformCollection = TOptional<FTransform>(),
 	bool bSetDefaultInternalMaterialsFromCollection = true,
-	FProgressCancel* Progress = nullptr
+	FProgressCancel* Progress = nullptr,
+	bool bSplitIslands = true
 );
 
 
@@ -473,6 +479,7 @@ int32 PLANARCUT_API AddCollisionSampleVertices(double TargetSpacing, FGeometryCo
  * @param CollisionSampleSpacing	Target spacing between collision sample vertices
  * @param TransformCollection		Optional transform of the collection; if unset, defaults to Identity
  * @param Progress					Optionally tracks progress and supports early-cancel
+ * @param bSplitIslands				Whether to detect connected components and split them apart after the cut
  * @return index of first new geometry in the Output GeometryCollection, or -1 if no geometry was added
  */
 int32 PLANARCUT_API CutWithMesh(
@@ -484,7 +491,8 @@ int32 PLANARCUT_API CutWithMesh(
 	double CollisionSampleSpacing,
 	const TOptional<FTransform>& TransformCollection = TOptional<FTransform>(),
 	bool bSetDefaultInternalMaterialsFromCollection = true,
-	FProgressCancel* Progress = nullptr
+	FProgressCancel* Progress = nullptr,
+	bool bSplitIslands = true
 );
 
 /// Convert a mesh description to a dynamic mesh *specifically* augmented / designed for cutting geometry collections, to be passed to CutWithMesh
@@ -499,7 +507,8 @@ inline int32 CutWithMesh(
 	double CollisionSampleSpacing,
 	const TOptional<FTransform>& TransformCollection = TOptional<FTransform>(),
 	bool bSetDefaultInternalMaterialsFromCollection = true,
-	FProgressCancel* Progress = nullptr
+	FProgressCancel* Progress = nullptr,
+	bool bSplitIslands = true
 )
 {
 	int32 NumUVLayers = Collection.NumUVLayers();
@@ -511,7 +520,8 @@ inline int32 CutWithMesh(
 		CollisionSampleSpacing,
 		TransformCollection,
 		bSetDefaultInternalMaterialsFromCollection,
-		Progress
+		Progress,
+		bSplitIslands
 		);
 }
 
