fix missing parameter from CutWithPlanarCells convenience function implementation, to fix fracture unit tests compilation

diff --git a/Engine/Plugins/Experimental/PlanarCutPlugin/Source/PlanarCut/Private/PlanarCut.cpp b/Engine/Plugins/Experimental/PlanarCutPlugin/Source/PlanarCut/Private/PlanarCut.cpp
--- a/Engine/Plugins/Experimental/PlanarCutPlugin/Source/PlanarCut/Private/PlanarCut.cpp
+++ b/Engine/Plugins/Experimental/PlanarCutPlugin/Source/PlanarCut/Private/PlanarCut.cpp
@@ -810,11 +810,12 @@ int32 CutWithPlanarCells(
 	bool bIncludeOutsideCellInOutput,
 	bool bSetDefaultInternalMaterialsFromCollection,
 	FProgressCancel* Progress,
-	FVector CellsOrigin
+	FVector CellsOrigin,
+	bool bSplitIslands
 )
 {
 	TArray<int32> TransformIndices { TransformIdx };
-	return CutMultipleWithPlanarCells(Cells, Source, TransformIndices, Grout, CollisionSampleSpacing, RandomSeed, TransformCollection, bIncludeOutsideCellInOutput, bSetDefaultInternalMaterialsFromCollection, Progress, CellsOrigin);
+	return CutMultipleWithPlanarCells(Cells, Source, TransformIndices, Grout, CollisionSampleSpacing, RandomSeed, TransformCollection, bIncludeOutsideCellInOutput, bSetDefaultInternalMaterialsFromCollection, Progress, CellsOrigin, bSplitIslands);
 }
 
 
