Fracture UV utilities : Add support to pass an array of selected faces to bake attribute from a geometry collection
- Added the possibility of having a mesh to collection face mnapping in the FDynamicMeshCollection
- Added a second TextureSpecifiedFaces that takes a bollean array for selected face
- Changed the code of TextureSpecifiedFaces to not rely on MaterialIds and type of face enums

diff --git a/Engine/Plugins/Experimental/PlanarCutPlugin/Source/PlanarCut/Private/FractureAutoUV.cpp b/Engine/Plugins/Experimental/PlanarCutPlugin/Source/PlanarCut/Private/FractureAutoUV.cpp
--- a/Engine/Plugins/Experimental/PlanarCutPlugin/Source/PlanarCut/Private/FractureAutoUV.cpp
+++ b/Engine/Plugins/Experimental/PlanarCutPlugin/Source/PlanarCut/Private/FractureAutoUV.cpp
@@ -667,7 +667,36 @@ bool TextureSpecifiedFaces(
 )
 {
 	TArray<bool> ToTextureTriangles;
-	int32 NumTextureTris = SetActiveTriangles(&Collection, ToTextureTriangles, true, TargetFaces, WhichMaterials);
+	SetActiveTriangles(&Collection, ToTextureTriangles, true, TargetFaces, WhichMaterials);
+
+	return TextureSpecifiedFaces(
+		TargetUVLayer,
+		Collection,
+		GutterSize,
+		BakeAttributes,
+		AttributeSettings,
+		TextureOut,
+		ToTextureTriangles,
+		Progress);
+}
+
+bool TextureSpecifiedFaces(
+	int32 TargetUVLayer,
+	FGeometryCollection& Collection,
+	int32 GutterSize,
+	FIndex4i BakeAttributes,
+	const FTextureAttributeSettings& AttributeSettings,
+	TImageBuilder<FVector4f>& TextureOut,
+	TArrayView<bool> ToTextureTriangles,
+	FProgressCancel* Progress
+)
+{
+	int32 NumTextureTris = 0;
+	for (const bool& bSelected : ToTextureTriangles)
+	{
+		NumTextureTris += (bSelected) ? 1 : 0;
+	}
+
 	FGeomFlatUVMesh UVMesh(TargetUVLayer, &Collection, ToTextureTriangles, NumTextureTris);
 
 	FImageOccupancyMap OccupancyMap;
@@ -717,31 +746,29 @@ bool TextureSpecifiedFaces(
 		float(CurvatureIdx > -1) * CurvatureWorkPer * NumMeshes;
 	const float WorkIncr = AmountOfWork > 0 ? 0.9f / AmountOfWork : 0.f;
 
-	FDynamicMeshCollection CollectionMeshes(&Collection, TransformIndices, FTransform::Identity, false);
+	FDynamicMeshCollection CollectionMeshes;
+	CollectionMeshes.bGenerateMeshToCollectionFaceMapping = true;
+	CollectionMeshes.Init(&Collection, TransformIndices, FTransform::Identity, false);
 	if (bNeedsDynamicMeshes)
 	{
-		TSet<int32> TargetMaterials;
-		for (int32 MatID : WhichMaterials)
-		{
-			TargetMaterials.Add(MatID);
-		}
-
-		// To bake only internal faces, unset texture on everything else
+		// To bake only selected faces, unset texture on everything else
 		for (int MeshIdx = 0; MeshIdx < CollectionMeshes.Meshes.Num(); MeshIdx++)
 		{
-			int32 TransformIdx = CollectionMeshes.Meshes[MeshIdx].TransformIndex;
-			FDynamicMesh3& Mesh = CollectionMeshes.Meshes[MeshIdx].AugMesh;
+			FDynamicMeshCollection::FMeshData& MeshData = CollectionMeshes.Meshes[MeshIdx];
+			int32 TransformIdx = MeshData.TransformIndex;
+			FDynamicMesh3& Mesh = MeshData.AugMesh;
 			FMeshNormals::InitializeOverlayToPerVertexNormals(Mesh.Attributes()->PrimaryNormals(), true);
 			AugmentedDynamicMesh::InitializeOverlayToPerVertexUVs(Mesh, 1, TargetUVLayer); // build an overlay w/ *only* the uv layer we care about
 			FDynamicMeshUVOverlay* UV = Mesh.Attributes()->PrimaryUV();
 			for (int TID : Mesh.TriangleIndicesItr())
 			{
-				bool bIsTextureTri = IsTriActive(
-					AugmentedDynamicMesh::GetVisibility(Mesh, TID), AugmentedDynamicMesh::GetInternal(Mesh, TID),
-					Mesh.Attributes()->GetMaterialID()->GetValue(TID), TargetMaterials, true, TargetFaces);
-				if (!bIsTextureTri)
+				const int32 CollectionFaceIndex = MeshData.GetCollectionFaceFromMeshFace(TID);
+				if (ensure(CollectionFaceIndex != INDEX_NONE && ToTextureTriangles.IsValidIndex(CollectionFaceIndex)))
 				{
-					UV->UnsetTriangle(TID);
+					if (ToTextureTriangles[CollectionFaceIndex] == false)
+					{
+						UV->UnsetTriangle(TID);
+					}
 				}
 			}
 		}
diff --git a/Engine/Plugins/Experimental/PlanarCutPlugin/Source/PlanarCut/Private/GeometryMeshConversion.cpp b/Engine/Plugins/Experimental/PlanarCutPlugin/Source/PlanarCut/Private/GeometryMeshConversion.cpp
--- a/Engine/Plugins/Experimental/PlanarCutPlugin/Source/PlanarCut/Private/GeometryMeshConversion.cpp
+++ b/Engine/Plugins/Experimental/PlanarCutPlugin/Source/PlanarCut/Private/GeometryMeshConversion.cpp
@@ -2636,6 +2636,12 @@ void FDynamicMeshCollection::InitTemplate(const FGeometryCollection* Collection,
 			{
 				continue;
 			}
+
+			if (bGenerateMeshToCollectionFaceMapping)
+			{
+				MeshData.AddMeshToCollectionFaceMapping(TID, Idx);
+			}
+
 			Mesh.Attributes()->GetMaterialID()->SetValue(TID, Collection->MaterialID[Idx]);
 			AugmentedDynamicMesh::SetVisibility(Mesh, TID, Collection->Visible[Idx]);
 			AugmentedDynamicMesh::SetInternal(Mesh, TID, Collection->Internal[Idx]);
diff --git a/Engine/Plugins/Experimental/PlanarCutPlugin/Source/PlanarCut/Private/GeometryMeshConversion.h b/Engine/Plugins/Experimental/PlanarCutPlugin/Source/PlanarCut/Private/GeometryMeshConversion.h
--- a/Engine/Plugins/Experimental/PlanarCutPlugin/Source/PlanarCut/Private/GeometryMeshConversion.h
+++ b/Engine/Plugins/Experimental/PlanarCutPlugin/Source/PlanarCut/Private/GeometryMeshConversion.h
@@ -197,9 +197,21 @@ struct FDynamicMeshCollection
 			bHasBounds = false;
 		}
 
+		void AddMeshToCollectionFaceMapping(int32 MeshFaceIndex, int32 CollectionFaceIndex)
+		{
+			MeshToCollectionFaceMapping.Add(MeshFaceIndex, CollectionFaceIndex);
+		}
+
+		int32 GetCollectionFaceFromMeshFace(int32 MeshFaceIndex) const
+		{
+			return MeshToCollectionFaceMapping.FindRef(MeshFaceIndex, INDEX_NONE);
+		}
+
 	private:
 		bool bHasBounds = false;
 		UE::Geometry::FAxisAlignedBox3d Bounds;
+		// optional face mapping from the mesh to the collection ( see bGenerateMeshToCollectionFaceMapping option )
+		TMap<int32, int32> MeshToCollectionFaceMapping;
 	};
 	TIndirectArray<FMeshData> Meshes;
 	UE::Geometry::FAxisAlignedBox3d Bounds;
@@ -210,6 +222,7 @@ struct FDynamicMeshCollection
 	bool bSkipInvisible = false;
 	// If false, Transforms passed to Init are interpreted as relative to the parent bone transform. If true, Transforms are all in the same 'global' / component-relative space
 	bool bComponentSpaceTransforms = false;
+	bool bGenerateMeshToCollectionFaceMapping = false;
 	
 	FDynamicMeshCollection() {}
 
diff --git a/Engine/Plugins/Experimental/PlanarCutPlugin/Source/PlanarCut/Public/FractureAutoUV.h b/Engine/Plugins/Experimental/PlanarCutPlugin/Source/PlanarCut/Public/FractureAutoUV.h
--- a/Engine/Plugins/Experimental/PlanarCutPlugin/Source/PlanarCut/Public/FractureAutoUV.h
+++ b/Engine/Plugins/Experimental/PlanarCutPlugin/Source/PlanarCut/Public/FractureAutoUV.h
@@ -212,6 +212,28 @@ bool PLANARCUT_API TextureSpecifiedFaces(
 	FProgressCancel* Progress = nullptr
 );
 
+/**
+ * Generate a texture for specified groups of faces based on chosen BakeAttributes and AttributeSettings
+ *
+ * @param TargetUVLayer		Which UV layer to take UV coordinates from when creating the new texture
+ * @param Collection		The collection to be create a new texture for
+ * @param GutterSize		Number of texels to fill outside of UV island borders (values are copied from nearest inside pt)
+ * @param BakeAttributes	Which attributes to bake into which color channel
+ * @param AttributeSettings	Settings for the BakeAttributes
+ * @param TextureOut		Texture to write to
+ * @param ToTextureTriangles boolean Array that indicate what face from the collection is to be baked
+ * @param WhichMaterials	If non-empty, apply texture to the listed material IDs
+ */
+bool PLANARCUT_API TextureSpecifiedFaces(
+	int32 TargetUVLayer,
+	FGeometryCollection& Collection,
+	int32 GutterSize,
+	UE::Geometry::FIndex4i BakeAttributes,
+	const FTextureAttributeSettings& AttributeSettings,
+	UE::Geometry::TImageBuilder<FVector4f>& TextureOut,
+	TArrayView<bool> ToTextureTriangles,
+	FProgressCancel* Progress = nullptr
+);
 
 UE_DEPRECATED(5.3, "Use TextureSpecifiedFaces instead")
 bool PLANARCUT_API TextureInternalSurfaces(
