fix fracture autouv handling of invalid target uv channel

diff --git a/Engine/Plugins/Experimental/PlanarCutPlugin/Source/PlanarCut/Private/FractureAutoUV.cpp b/Engine/Plugins/Experimental/PlanarCutPlugin/Source/PlanarCut/Private/FractureAutoUV.cpp
--- a/Engine/Plugins/Experimental/PlanarCutPlugin/Source/PlanarCut/Private/FractureAutoUV.cpp
+++ b/Engine/Plugins/Experimental/PlanarCutPlugin/Source/PlanarCut/Private/FractureAutoUV.cpp
@@ -1,5 +1,6 @@
 // Copyright Epic Games, Inc. All Rights Reserved.
 #include "FractureAutoUV.h"
+#include "PlanarCutPlugin.h"
 
 #include "Distance/DistLine3Segment3.h"
 #include "GeometryMeshConversion.h"
@@ -194,7 +195,7 @@ struct FGeomMesh : public UE::Geometry::FUVPacker::IUVMeshView
 	{
 		if (!ensure(UVLayer >= 0 && UVLayer < Collection->NumUVLayers()))
 		{
-			UVLayer = FMath::Clamp(UVLayer, 0, Collection->NumUVLayers());
+			UVLayer = FMath::Clamp(UVLayer, 0, Collection->NumUVLayers() - 1);
 			CollectionUVs = Collection->FindUVLayer(UVLayer);
 		}
 		check(CollectionUVs);
@@ -495,6 +496,13 @@ bool PLANARCUT_API UVLayout(
 	FProgressCancel* Progress
 )
 {
+	// TargetUVLayer must be valid for the collection
+	if (TargetUVLayer < 0 || TargetUVLayer >= GeometryCollectionUV::MAX_NUM_UV_CHANNELS || !Collection.FindUVLayer(TargetUVLayer))
+	{
+		UE_LOG(LogPlanarCut, Error, TEXT("UV Layout failed: Target UV Channel (%d) not found on collection"), TargetUVLayer);
+		return false;
+	}
+
 	FProgressCancel::FProgressScope InitScope = FProgressCancel::CreateScopeTo(Progress, .2);
 
 	int32 NumActive = 0;
