Add option for TinyGeo to not rely on the geometry collection proximity settings
+ auto-update proximity if tiny geo does need it, so we know it's not stale (only in the fractureengine path)

diff --git a/Engine/Plugins/Experimental/Fracture/Source/FractureEngine/Private/FractureEngineUtility.cpp b/Engine/Plugins/Experimental/Fracture/Source/FractureEngine/Private/FractureEngineUtility.cpp
--- a/Engine/Plugins/Experimental/Fracture/Source/FractureEngine/Private/FractureEngineUtility.cpp
+++ b/Engine/Plugins/Experimental/Fracture/Source/FractureEngine/Private/FractureEngineUtility.cpp
@@ -233,7 +233,8 @@ void FFractureEngineUtility::FixTinyGeo(FManagedArrayCollection& InOutCollection
 	const bool InOnlyClusters,
 	const EFixTinyGeoNeighborSelectionMethod InNeighborSelection,
 	const bool InOnlyToConnected,
-	const bool InOnlySameParent)
+	const bool InOnlySameParent,
+	const bool bUseCollectionProximity)
 {
 	if (TUniquePtr<FGeometryCollection> GeomCollection = TUniquePtr<FGeometryCollection>(InOutCollection.NewCopy<FGeometryCollection>()))
 	{
@@ -262,6 +263,13 @@ void FFractureEngineUtility::FixTinyGeo(FManagedArrayCollection& InOutCollection
 				SelectionMethod = UE::PlanarCut::ENeighborSelectionMethod::NearestCenter;
 			}
 
+			// Make sure we have non-stale proximity data if we will use it
+			if (bUseCollectionProximity)
+			{
+				FGeometryCollectionProximityUtility ProximityUtility(GeomCollection.Get());
+				ProximityUtility.UpdateProximity();
+			}
+
 			if (InMergeType == EFixTinyGeoMergeType::MergeGeometry)
 			{
 				MergeBones(*GeomCollection,
@@ -270,7 +278,8 @@ void FFractureEngineUtility::FixTinyGeo(FManagedArrayCollection& InOutCollection
 					MinVolume,
 					SmallIndices,
 					false, /*bUnionJoinedPieces*/ // Note: Union-ing the pieces is nicer in theory, but can leave cracks and non-manifold garbage
-					SelectionMethod);
+					SelectionMethod,
+					bUseCollectionProximity);
 			}
 			else
 			{
@@ -280,7 +289,8 @@ void FFractureEngineUtility::FixTinyGeo(FManagedArrayCollection& InOutCollection
 					SmallIndices,
 					SelectionMethod,
 					InOnlyToConnected,
-					InOnlySameParent);
+					InOnlySameParent,
+					bUseCollectionProximity);
 			}
 
 			InOutCollection = (const FManagedArrayCollection&)(*GeomCollection);
diff --git a/Engine/Plugins/Experimental/Fracture/Source/FractureEngine/Public/FractureEngineUtility.h b/Engine/Plugins/Experimental/Fracture/Source/FractureEngine/Public/FractureEngineUtility.h
--- a/Engine/Plugins/Experimental/Fracture/Source/FractureEngine/Public/FractureEngineUtility.h
+++ b/Engine/Plugins/Experimental/Fracture/Source/FractureEngine/Public/FractureEngineUtility.h
@@ -73,7 +73,8 @@ public:
 		const bool InOnlyClusters,
 		const EFixTinyGeoNeighborSelectionMethod InNeighborSelection,
 		const bool InOnlyToConnected,
-		const bool InOnlySameParent);
+		const bool InOnlySameParent,
+		const bool bUseCollectionProximity = true);
 
 	static void RecomputeNormalsInGeometryCollection(FManagedArrayCollection& InOutCollection,
 		FDataflowTransformSelection InTransformSelection,
diff --git a/Engine/Plugins/Experimental/GeometryCollectionPlugin/Source/GeometryCollectionNodes/Private/Dataflow/GeometryCollectionUtilityNodes.cpp b/Engine/Plugins/Experimental/GeometryCollectionPlugin/Source/GeometryCollectionNodes/Private/Dataflow/GeometryCollectionUtilityNodes.cpp
--- a/Engine/Plugins/Experimental/GeometryCollectionPlugin/Source/GeometryCollectionNodes/Private/Dataflow/GeometryCollectionUtilityNodes.cpp
+++ b/Engine/Plugins/Experimental/GeometryCollectionPlugin/Source/GeometryCollectionNodes/Private/Dataflow/GeometryCollectionUtilityNodes.cpp
@@ -645,7 +645,8 @@ void FFixTinyGeoDataflowNode::Evaluate(UE::Dataflow::FContext& Context, const FD
 				bOnlyClusters,
 				NeighborSelection,
 				bOnlyToConnected,
-				bOnlySameParent);
+				bOnlySameParent,
+				bUseCollectionProximityForConnections);
 
 			SetValue(Context, MoveTemp(InCollection), &Collection);
 
diff --git a/Engine/Plugins/Experimental/GeometryCollectionPlugin/Source/GeometryCollectionNodes/Public/Dataflow/GeometryCollectionUtilityNodes.h b/Engine/Plugins/Experimental/GeometryCollectionPlugin/Source/GeometryCollectionNodes/Public/Dataflow/GeometryCollectionUtilityNodes.h
--- a/Engine/Plugins/Experimental/GeometryCollectionPlugin/Source/GeometryCollectionNodes/Public/Dataflow/GeometryCollectionUtilityNodes.h
+++ b/Engine/Plugins/Experimental/GeometryCollectionPlugin/Source/GeometryCollectionNodes/Public/Dataflow/GeometryCollectionUtilityNodes.h
@@ -616,6 +616,10 @@ public:
 	UPROPERTY(EditAnywhere, Category = MergeSettings, meta = (EditCondition = "MergeType == EFixTinyGeoMergeType::MergeClusters", EditConditionHides))
 	bool bOnlyToConnected = true;
 
+	/** Whether to use the Proximity (as computed by the Proximity node) to determine which bones are connected, and thus can be considered for merging. Otherwise will compute and use a reasonable default connectivity. */
+	UPROPERTY(EditAnywhere, Category = MergeSettings, AdvancedDisplay)
+	bool bUseCollectionProximityForConnections = false;
+
 	/** Options for using the current bone selection */
 	UPROPERTY(EditAnywhere, Category = MergeSettings, meta = (DisplayName = "Bone Selection"))
 	EFixTinyGeoUseBoneSelection UseBoneSelection = EFixTinyGeoUseBoneSelection::NoEffect;
diff --git a/Engine/Plugins/Experimental/PlanarCutPlugin/Source/PlanarCut/Private/PlanarCut.cpp b/Engine/Plugins/Experimental/PlanarCutPlugin/Source/PlanarCut/Private/PlanarCut.cpp
--- a/Engine/Plugins/Experimental/PlanarCutPlugin/Source/PlanarCut/Private/PlanarCut.cpp
+++ b/Engine/Plugins/Experimental/PlanarCutPlugin/Source/PlanarCut/Private/PlanarCut.cpp
@@ -1318,14 +1318,25 @@ int32 MergeBones(
 	double MinVolume,
 	const TArrayView<const int32>& SmallTransformIndices,
 	bool bUnionJoinedPieces,
-	UE::PlanarCut::ENeighborSelectionMethod NeighborSelectionMethod
+	UE::PlanarCut::ENeighborSelectionMethod NeighborSelectionMethod,
+	bool bUseCollectionProximity
 )
 {
 	FTransform CellsToWorld = FTransform::Identity;
 
 	FGeometryCollectionProximityUtility ProximityUtility(&Collection);
-	ProximityUtility.RequireProximity();
-	const TManagedArray<TSet<int32>>& Proximity = Collection.GetAttribute<TSet<int32>>("Proximity", FGeometryCollection::GeometryGroup);
+	TArray<TSet<int32>> LocalProximity;
+	const TArray<TSet<int32>>* UseProximity = nullptr;
+	if (!bUseCollectionProximity)
+	{
+		LocalProximity = ProximityUtility.ComputePreciseProximity(Collection);
+		UseProximity = &LocalProximity;
+	}
+	else
+	{
+		ProximityUtility.RequireProximity();
+		UseProximity = &Collection.GetAttribute<TSet<int32>>("Proximity", FGeometryCollection::GeometryGroup).GetConstArray();
+	}
 
 	// local array so we can populate it with all transforms if input was empty
 	TArray<int32> TransformIndices(TransformIndicesView);
@@ -1488,7 +1499,7 @@ int32 MergeBones(
 			}
 		}
 
-		const TSet<int32>& Prox = Proximity[SmallIdx];
+		const TSet<int32>& Prox = (*UseProximity)[SmallIdx];
 		double BestScore = -FMathd::MaxReal;
 		int32 BestNbrIdx = INDEX_NONE;
 		for (int32 NbrIdx : Prox)
@@ -1697,7 +1708,8 @@ void MergeClusters(
 	const TArrayView<const int32>& SmallTransformIndices,
 	UE::PlanarCut::ENeighborSelectionMethod NeighborSelectionMethod,
 	bool bOnlyMergeInProximity,
-	bool bOnlySameParent
+	bool bOnlySameParent,
+	bool bUseCollectionProximity
 )
 {
 	if (Collection.NumElements(FGeometryCollection::TransformGroup) == 0)
@@ -1707,8 +1719,18 @@ void MergeClusters(
 	}
 
 	FGeometryCollectionProximityUtility ProximityUtility(&Collection);
-	ProximityUtility.RequireProximity();
-	const TManagedArray<TSet<int32>>& Proximity = Collection.GetAttribute<TSet<int32>>("Proximity", FGeometryCollection::GeometryGroup);
+	TArray<TSet<int32>> LocalProximity;
+	const TArray<TSet<int32>>* UseProximity = nullptr;
+	if (!bUseCollectionProximity)
+	{
+		LocalProximity = ProximityUtility.ComputePreciseProximity(Collection);
+		UseProximity = &LocalProximity;
+	}
+	else
+	{
+		ProximityUtility.RequireProximity();
+		UseProximity = &Collection.GetAttribute<TSet<int32>>("Proximity", FGeometryCollection::GeometryGroup).GetConstArray();
+	}
 	const TManagedArray<int32>* Level = Collection.FindAttribute<int32>("Level", FGeometryCollection::TransformGroup);
 	if (!ensure(Level != nullptr))
 	{
@@ -1737,7 +1759,7 @@ void MergeClusters(
 	// only selects rigid or cluster bones
 	// Note we do the traversal each time rather using a pre-computed connection graph because the cluster neighbors will change as we merge clusters
 	// (the Proximity attribute should be stable because it is on the geometry group, not the transform group)
-	auto GetSameLevelNeighbors = [&Collection, &Proximity, &Level, bOnlySameParent](int32 BoneIdx, TSet<int32>& OutNeighbors) -> void
+	auto GetSameLevelNeighbors = [&Collection, &UseProximity, &Level, bOnlySameParent](int32 BoneIdx, TSet<int32>& OutNeighbors) -> void
 	{
 		int32 TargetLevel = (*Level)[BoneIdx];
 		TArray<int32> ToProcess;
@@ -1750,7 +1772,7 @@ void MergeClusters(
 			if (Collection.SimulationType[ProcessIdx] == FGeometryCollection::ESimulationTypes::FST_Rigid
 				&& Collection.TransformToGeometryIndex[ProcessIdx] != INDEX_NONE)
 			{
-				for (int32 NbrGeo : Proximity[Collection.TransformToGeometryIndex[ProcessIdx]])
+				for (int32 NbrGeo : (*UseProximity)[Collection.TransformToGeometryIndex[ProcessIdx]])
 				{
 					// Traverse parents to find a same-level cluster on this neighbor
 					int32 TraverseParent = Collection.TransformIndex[NbrGeo];
diff --git a/Engine/Plugins/Experimental/PlanarCutPlugin/Source/PlanarCut/Public/PlanarCut.h b/Engine/Plugins/Experimental/PlanarCutPlugin/Source/PlanarCut/Public/PlanarCut.h
--- a/Engine/Plugins/Experimental/PlanarCutPlugin/Source/PlanarCut/Public/PlanarCut.h
+++ b/Engine/Plugins/Experimental/PlanarCutPlugin/Source/PlanarCut/Public/PlanarCut.h
@@ -400,6 +400,7 @@ namespace UE
  * @param SmallTransformIndices		Transform indices of pieces that we want to merge
  * @param bUnionJoinedPieces		Try to 'union' the merged pieces, removing internal triangles and connecting the shared cut boundary
  * @param NeighborSelectionMethod	How to choose which neighbor to merge to
+ * @param bUseCollectionProximity	If true, will use the proximity data on the collection to decide which pieces are allowed to merge. Otherwise, a reasonable default proximity will be computed and used.
  */
 int32 PLANARCUT_API MergeBones(
 	FGeometryCollection& Collection,
@@ -408,7 +409,8 @@ int32 PLANARCUT_API MergeBones(
 	double MinVolume,
 	const TArrayView<const int32>& SmallTransformIndices,
 	bool bUnionJoinedPieces,
-	UE::PlanarCut::ENeighborSelectionMethod NeighborSelectionMethod
+	UE::PlanarCut::ENeighborSelectionMethod NeighborSelectionMethod,
+	bool bUseCollectionProximity = true
 );
 
 /**
@@ -421,6 +423,7 @@ int32 PLANARCUT_API MergeBones(
  * @param NeighborSelectionMethod	How to choose which neighbor to merge to
  * @param bOnlyMergeInProximity		Whether to only follow proximity graph, or to fall back to nearest center when for when no proximity links are found
  * @param bOnlySameParent			Whether to only merge nodes that have the same parent in the hierarchy
+ * @param bUseCollectionProximity	If true, will use the proximity data on the collection to decide which pieces are allowed to merge. Otherwise, a reasonable default proximity will be computed and used.
  */
 void PLANARCUT_API MergeClusters(
 	FGeometryCollection& Collection,
@@ -429,7 +432,8 @@ void PLANARCUT_API MergeClusters(
 	const TArrayView<const int32>& SmallTransformIndices,
 	UE::PlanarCut::ENeighborSelectionMethod NeighborSelectionMethod,
 	bool bOnlyMergeInProximity,
-	bool bOnlySameParent
+	bool bOnlySameParent,
+	bool bUseCollectionProximity = true
 );
 
 /**
diff --git a/Engine/Source/Runtime/Experimental/Chaos/Private/GeometryCollection/GeometryCollectionProximityUtility.cpp b/Engine/Source/Runtime/Experimental/Chaos/Private/GeometryCollection/GeometryCollectionProximityUtility.cpp
--- a/Engine/Source/Runtime/Experimental/Chaos/Private/GeometryCollection/GeometryCollectionProximityUtility.cpp
+++ b/Engine/Source/Runtime/Experimental/Chaos/Private/GeometryCollection/GeometryCollectionProximityUtility.cpp
@@ -239,6 +239,23 @@ struct FGeometryCollectionProximitySpatial
 
 	}
 
+	TArray<TSet<int32>> MoveProximityToArray()
+	{
+		TArray<TSet<int32>> Proximity = MoveTemp(KnownProximity);
+		for (int32 GeoIdx = 0; GeoIdx < Proximity.Num(); ++GeoIdx)
+		{
+			for (const TPair<int32, bool>& PossibleContact : GeoInfo[GeoIdx].CandidateContacts)
+			{
+				if (PossibleContact.Value)
+				{
+					Proximity[GeoIdx].Add(PossibleContact.Key);
+					Proximity[PossibleContact.Key].Add(GeoIdx);
+				}
+			}
+		}
+		return Proximity;
+	}
+
 	bool AreFacesCoPlanar(const FGeometryCollection* Collection, int32 Idx0, int32 Idx1, float ProximityTolerance) const
 	{
 		// Assumes that faces have already been determined to be parallel.
@@ -697,6 +714,13 @@ void FGeometryCollectionProximityUtility::CopyProximityToConnectionGraph(const T
 	}
 }
 
+TArray<TSet<int32>> FGeometryCollectionProximityUtility::ComputePreciseProximity(const FGeometryCollection& Collection)
+{
+	using namespace UE::GeometryCollectionInternal;
+	FGeometryCollectionProximitySpatial Spatial(&Collection, PreciseProximityThreshold);
+	return Spatial.MoveProximityToArray();
+}
+
 void FGeometryCollectionProximityUtility::UpdateProximity(UE::GeometryCollectionConvexUtility::FConvexHulls* OptionalComputedHulls)
 {
 	using namespace UE::GeometryCollectionInternal;
@@ -728,9 +752,6 @@ void FGeometryCollectionProximityUtility::UpdateProximity(UE::GeometryCollection
 	}
 	else
 	{
-		// This threshold not exposed to the user via Properties.DistanceThreshold because it doesn't behave very well if increased to large values --
-		// the computation would become closer to O(n^2), and the results would likely be confusing/inconsistent
-		constexpr float PreciseProximityThreshold = .01f;
 		FGeometryCollectionProximitySpatial Spatial(Collection, PreciseProximityThreshold);
 		Spatial.MoveProximityToCollection(Collection);
 	}
diff --git a/Engine/Source/Runtime/Experimental/Chaos/Public/GeometryCollection/GeometryCollectionProximityUtility.h b/Engine/Source/Runtime/Experimental/Chaos/Public/GeometryCollection/GeometryCollectionProximityUtility.h
--- a/Engine/Source/Runtime/Experimental/Chaos/Public/GeometryCollection/GeometryCollectionProximityUtility.h
+++ b/Engine/Source/Runtime/Experimental/Chaos/Public/GeometryCollection/GeometryCollectionProximityUtility.h
@@ -63,6 +63,9 @@ public:
 
 	CHAOS_API void EnumerateNeighbors(const Chaos::Facades::FCollectionHierarchyFacade& Hierarchy, int32 BoneIdx, TFunctionRef<void(int32)> NeighborFunc, bool bIncludeNeighborsInParentLevels = true, bool bFilterDuplicates = true);
 
+	// Method to compute the precise proximity, regardless of the proximity configuration on the collection
+	CHAOS_API static TArray<TSet<int32>> ComputePreciseProximity(const FGeometryCollection& InCollection);
+
 	// Stores stats about the contact between two geometries
 	struct FGeometryContactEdge
 	{
@@ -90,5 +93,9 @@ public:
 
 private:
 	FGeometryCollection* Collection;
+
+	// This threshold not exposed to the user via Properties.DistanceThreshold because it doesn't behave very well if increased to large values --
+	// the computation would become closer to O(n^2), and the results would likely be confusing/inconsistent
+	static constexpr float PreciseProximityThreshold = .01f;
 };
 
