fix for crash on out of sync selection input to fracture auto uv node

diff --git a/Engine/Plugins/Experimental/PlanarCutPlugin/Source/PlanarCut/Private/FractureAutoUV.cpp b/Engine/Plugins/Experimental/PlanarCutPlugin/Source/PlanarCut/Private/FractureAutoUV.cpp
--- a/Engine/Plugins/Experimental/PlanarCutPlugin/Source/PlanarCut/Private/FractureAutoUV.cpp
+++ b/Engine/Plugins/Experimental/PlanarCutPlugin/Source/PlanarCut/Private/FractureAutoUV.cpp
@@ -502,6 +502,12 @@ bool PLANARCUT_API UVLayout(
 		UE_LOG(LogPlanarCut, Error, TEXT("UV Layout failed: Target UV Channel (%d) not found on collection"), TargetUVLayer);
 		return false;
 	}
+	// Face selection array size must match the collection face count
+	if (Collection.Indices.Num() != FaceSelection.Num())
+	{
+		UE_LOG(LogPlanarCut, Error, TEXT("UV Layout failed: Mismatch between collection face count (%d) and selection face count (%d)"), Collection.Indices.Num(), FaceSelection.Num());
+		return false;
+	}
 
 	FProgressCancel::FProgressScope InitScope = FProgressCancel::CreateScopeTo(Progress, .2);
 
