add method to merge UV islands on geometry collections + expose as a new dataflow node

diff --git a/Engine/Plugins/Experimental/GeometryCollectionPlugin/Source/GeometryCollectionNodes/Private/Dataflow/GeometryCollectionUVNodes.cpp b/Engine/Plugins/Experimental/GeometryCollectionPlugin/Source/GeometryCollectionNodes/Private/Dataflow/GeometryCollectionUVNodes.cpp
--- a/Engine/Plugins/Experimental/GeometryCollectionPlugin/Source/GeometryCollectionNodes/Private/Dataflow/GeometryCollectionUVNodes.cpp
+++ b/Engine/Plugins/Experimental/GeometryCollectionPlugin/Source/GeometryCollectionNodes/Private/Dataflow/GeometryCollectionUVNodes.cpp
@@ -11,6 +11,7 @@ namespace UE::Dataflow
 	{
 		DATAFLOW_NODE_REGISTER_CREATION_FACTORY(FAddUVChannelDataflowNode);
 		DATAFLOW_NODE_REGISTER_CREATION_FACTORY(FAutoUnwrapUVDataflowNode);
+		DATAFLOW_NODE_REGISTER_CREATION_FACTORY(FMergeUVIslandsDataflowNode);
 		DATAFLOW_NODE_REGISTER_CREATION_FACTORY(FBoxProjectUVDataflowNode);
 	}
 }
@@ -62,6 +63,72 @@ void FAddUVChannelDataflowNode::Evaluate(UE::Dataflow::FContext& Context, const
 
 ////////////////////////////////////////////////////////////////////////////////////////////
 
+FMergeUVIslandsDataflowNode::FMergeUVIslandsDataflowNode(const UE::Dataflow::FNodeParameters& InParam, FGuid InGuid)
+	: FDataflowNode(InParam, InGuid)
+{
+	RegisterInputConnection(&Collection);
+	RegisterInputConnection(&FaceSelection);
+	RegisterInputConnection(&UVChannel);
+	RegisterInputConnection(&AreaDistortionThreshold).SetCanHidePin(true).SetPinIsHidden(true);
+	RegisterInputConnection(&MaxNormalDeviationDeg).SetCanHidePin(true).SetPinIsHidden(true);
+	RegisterInputConnection(&NormalSmoothingRounds).SetCanHidePin(true).SetPinIsHidden(true);
+	RegisterInputConnection(&NormalSmoothingAlpha).SetCanHidePin(true).SetPinIsHidden(true);
+	RegisterOutputConnection(&Collection, &Collection);
+}
+
+void FMergeUVIslandsDataflowNode::Evaluate(UE::Dataflow::FContext& Context, const FDataflowOutput* Out) const
+{
+	if (Out->IsA(&Collection))
+	{
+		const FManagedArrayCollection& InCollection = GetValue(Context, &Collection);
+		const int32 FaceCount = InCollection.NumElements(FGeometryCollection::FacesGroup);
+		if (FaceCount == 0)
+		{
+			// nothing to do : forward and exit 
+			SafeForwardInput(Context, &Collection, &Collection);
+			return;
+		}
+
+		const int32 InUVChannel = GetValue(Context, &UVChannel);
+		UE::PlanarCut::FMergeIslandSettings MergeIslandSettings;
+		MergeIslandSettings.MaxNormalDeviationDeg = GetValue(Context, &MaxNormalDeviationDeg);
+		MergeIslandSettings.AreaDistortionThreshold = GetValue(Context, &AreaDistortionThreshold);
+		MergeIslandSettings.NormalSmoothingAlpha = GetValue(Context, &NormalSmoothingAlpha);
+		MergeIslandSettings.NormalSmoothingRounds = GetValue(Context, &NormalSmoothingRounds);
+
+		TUniquePtr<FGeometryCollection> GeomCollection = TUniquePtr<FGeometryCollection>(InCollection.NewCopy<FGeometryCollection>());
+		check(GeomCollection.IsValid());
+
+		if (IsConnected(&FaceSelection))
+		{
+			const FDataflowFaceSelection& InFaceSelection = GetValue(Context, &FaceSelection);
+			if (InFaceSelection.Num() > 0)
+			{
+				TArray<bool> FacesToUnwrap;
+				FacesToUnwrap.SetNumUninitialized(InFaceSelection.Num());
+				for (int32 Index = 0; Index < InFaceSelection.Num(); ++Index)
+				{
+					FacesToUnwrap[Index] = InFaceSelection.IsSelected(Index);
+				}
+
+				UE::PlanarCut::MergeUVIslands(InUVChannel, *GeomCollection, MergeIslandSettings,
+					FacesToUnwrap, nullptr /* progress */);
+			}
+		}
+		else
+		{
+			// All visible faces
+			TArray<bool> FacesToUnwrap = GeomCollection->Visible.GetAsBoolArray();
+			UE::PlanarCut::MergeUVIslands(InUVChannel, *GeomCollection, MergeIslandSettings,
+				FacesToUnwrap, nullptr /* progress */);
+		}
+
+		SetValue(Context, (const FManagedArrayCollection&)(*GeomCollection), &Collection);
+	}
+}
+
+////////////////////////////////////////////////////////////////////////////////////////////
+
 FAutoUnwrapUVDataflowNode::FAutoUnwrapUVDataflowNode(const UE::Dataflow::FNodeParameters& InParam, FGuid InGuid)
 	: FDataflowNode(InParam, InGuid)
 {
diff --git a/Engine/Plugins/Experimental/GeometryCollectionPlugin/Source/GeometryCollectionNodes/Public/Dataflow/GeometryCollectionUVNodes.h b/Engine/Plugins/Experimental/GeometryCollectionPlugin/Source/GeometryCollectionNodes/Public/Dataflow/GeometryCollectionUVNodes.h
--- a/Engine/Plugins/Experimental/GeometryCollectionPlugin/Source/GeometryCollectionNodes/Public/Dataflow/GeometryCollectionUVNodes.h
+++ b/Engine/Plugins/Experimental/GeometryCollectionPlugin/Source/GeometryCollectionNodes/Public/Dataflow/GeometryCollectionUVNodes.h
@@ -69,6 +69,49 @@ public:
 	FAutoUnwrapUVDataflowNode(const UE::Dataflow::FNodeParameters& InParam, FGuid InGuid = FGuid::NewGuid());
 };
 
+/**
+ * Merge adjacent UV Islands with similar normals for a specific UV channel
+ */
+USTRUCT()
+struct FMergeUVIslandsDataflowNode : public FDataflowNode
+{
+	GENERATED_USTRUCT_BODY()
+	DATAFLOW_NODE_DEFINE_INTERNAL(FMergeUVIslandsDataflowNode, "MergeUVIslands", "GeometryCollection|UV", "")
+
+private:
+	/** Target Collection */
+	UPROPERTY(meta = (DataflowInput, DataflowOutput, DataflowPassthrough = "Collection", DataflowIntrinsic))
+	FManagedArrayCollection Collection;
+
+	/** Faces to auto unwrap, no selection means all faces */
+	UPROPERTY(meta = (DataflowInput))
+	FDataflowFaceSelection FaceSelection;
+
+	/** UV channel to unwrap into ( 0 by default ) */
+	UPROPERTY(EditAnywhere, Category = Options, meta = (DataflowInput, ClampMin = 0, ClampMax = 7, UIMin = 0, UIMax = 7))
+	int32 UVChannel = 0;
+
+
+	// Threshold for allowed area distortion from merging islands (when we use ExpMap to compute new UVs for the merged island)
+	UPROPERTY(EditAnywhere, Category = Options, meta = (DataflowInput, ClampMin = 0, ClampMax = 10, UIMin = 0, UIMax = 10))
+	double AreaDistortionThreshold = 1.5;
+	// Threshold for allowed normal deviation between merge-able islands
+	UPROPERTY(EditAnywhere, Category = Options, meta = (DataflowInput, ClampMin = 0, ClampMax = 180, UIMin = 0, UIMax = 90))
+	double MaxNormalDeviationDeg = 45.0;
+
+	// Amount of normal smoothing to apply when computing new UVs for merged islands. More smoothing will result in UV maps that are less sensitive to local surface shape.
+	UPROPERTY(EditAnywhere, Category = Options, AdvancedDisplay, meta = (DataflowInput, ClampMin = 0, ClampMax = 500, UIMin = 0, UIMax = 100))
+	int32 NormalSmoothingRounds = 0;
+	// Strength of normal smoothing to apply when computing new UVs for merged islands. Stronger smoothing will result in UV maps that are less sensitive to local surface shape.
+	UPROPERTY(EditAnywhere, Category = Options, AdvancedDisplay, meta = (DataflowInput, ClampMin = 0, ClampMax = 1, UIMin = 0, UIMax = 1))
+	double NormalSmoothingAlpha = 0.25;
+
+	virtual void Evaluate(UE::Dataflow::FContext& Context, const FDataflowOutput* Out) const override;
+
+public:
+	FMergeUVIslandsDataflowNode(const UE::Dataflow::FNodeParameters& InParam, FGuid InGuid = FGuid::NewGuid());
+};
+
 /**
  * Generates UVs using a box projection
  */
diff --git a/Engine/Plugins/Experimental/PlanarCutPlugin/Source/PlanarCut/Private/FractureAutoUV.cpp b/Engine/Plugins/Experimental/PlanarCutPlugin/Source/PlanarCut/Private/FractureAutoUV.cpp
--- a/Engine/Plugins/Experimental/PlanarCutPlugin/Source/PlanarCut/Private/FractureAutoUV.cpp
+++ b/Engine/Plugins/Experimental/PlanarCutPlugin/Source/PlanarCut/Private/FractureAutoUV.cpp
@@ -18,6 +18,7 @@
 
 #include "Implicit/Solidify.h"
 
+#include "Parameterization/PatchBasedMeshUVGenerator.h"
 #include "Parameterization/UVPacking.h"
 #include "Sampling/MeshCurvatureMapBaker.h"
 #include "Sampling/MeshOcclusionMapBaker.h"
@@ -102,6 +103,63 @@ namespace UE
 	}
 } // namespace UE::UVPacking
 
+namespace UE::PrivateUVHelpers
+{
+	// Helper to update geometry UVs for meshes in specified bones via some processing that operates on a dynamic mesh
+	bool ApplyDynamicMeshBasedUVProcessing(
+		int32 TargetUVLayer, FGeometryCollection& Collection, 
+		TArrayView<const int32> TransformIndices, TFunctionRef<bool(FDynamicMesh3&, int32 TransformIndex)> ProcessMesh,
+		bool bWeldUVs = true, bool bSaveIsolatedVertices = true, bool bNeedsSplitAttributes = true)
+	{
+		using namespace UE::PlanarCut;
+
+		int32 NumUVLayers = Collection.NumUVLayers();
+		if (TargetUVLayer < 0 || TargetUVLayer >= NumUVLayers)
+		{
+			return false;
+		}
+
+		FDynamicMeshCollection CollectionMeshes(&Collection, TransformIndices, FTransform::Identity, bSaveIsolatedVertices);
+
+		bool bAnyUpdate = false;
+
+		for (int MeshIdx = 0; MeshIdx < CollectionMeshes.Meshes.Num(); MeshIdx++)
+		{
+			int32 TransformIdx = CollectionMeshes.Meshes[MeshIdx].TransformIndex;
+			FDynamicMesh3& Mesh = CollectionMeshes.Meshes[MeshIdx].AugMesh;
+
+			FMeshNormals::InitializeOverlayToPerVertexNormals(Mesh.Attributes()->PrimaryNormals(), true);
+			AugmentedDynamicMesh::InitializeOverlayToPerVertexTangents(Mesh);
+			AugmentedDynamicMesh::InitializeOverlayToPerVertexUVs(Mesh, NumUVLayers, 0);
+
+			// Optionally weld edges so we won't have artificial UV seams on sharp edges, etc
+			if (bWeldUVs)
+			{
+				FMergeCoincidentMeshEdges EdgeWelder(&Mesh);
+				EdgeWelder.Apply();
+			}
+
+			if (ProcessMesh(Mesh, TransformIdx))
+			{
+				bAnyUpdate = true;
+				Mesh.CompactInPlace();
+				// Geometry collection requires split attributes, so we always need to re-split them if we've welded above.
+				// We also re-split them if requested (i.e., if not welding, set this if ProcessMesh may create un-split vertices)
+				if (bWeldUVs || bNeedsSplitAttributes)
+				{
+					AugmentedDynamicMesh::SplitOverlayAttributesToPerVertex(Mesh, true, true);
+				}
+			}
+		}
+
+		if (bAnyUpdate)
+		{
+			CollectionMeshes.UpdateAllCollections(Collection);
+		}
+
+		return bAnyUpdate;
+	}
+}
 
 namespace UE { namespace PlanarCut {
 
@@ -376,6 +434,114 @@ struct FGeomFlatUVMesh
 };
 
 
+bool MergeUVIslands(
+	int32 TargetUVLayer,
+	FGeometryCollection& Collection,
+	FMergeIslandSettings MergeIslandSettings,
+	TArrayView<bool> FaceSelection,
+	FProgressCancel* Progress
+)
+{
+	TSet<int32> ActiveTransformIndexSet;
+	for (int32 FaceIdx = 0; FaceIdx < FaceSelection.Num(); ++FaceIdx)
+	{
+		ActiveTransformIndexSet.Add(Collection.BoneMap[Collection.Indices[FaceIdx].X]);
+	}
+	TArray<int32> TransformIndices = ActiveTransformIndexSet.Array();
+
+	return UE::PrivateUVHelpers::ApplyDynamicMeshBasedUVProcessing(
+		TargetUVLayer, Collection,
+		TransformIndices, [&MergeIslandSettings, &FaceSelection, &Collection, TargetUVLayer](FDynamicMesh3& Mesh, int32 TransformIndex) -> bool
+		{
+
+			FDynamicMeshUVOverlay* UVLayer = Mesh.Attributes()->GetUVLayer(TargetUVLayer);
+			if (!ensure(UVLayer))
+			{
+				return false;
+			}
+
+			// Pull the selected faces for the current mesh from the overall FaceSelection, via the mesh TransformIndex
+			int32 GeometryIndex = Collection.TransformToGeometryIndex[TransformIndex];
+			int32 FaceStart = Collection.FaceStart[GeometryIndex];
+			// Note: Mesh must have the same triangle count as the corresponding geometry, and the triangles must be compact
+			if (!ensure(Collection.FaceCount[GeometryIndex] == Mesh.TriangleCount() && Mesh.IsCompactT()))
+			{
+				return false;
+			}
+			TArray<int32> TargetTris;
+			for (int TID : Mesh.TriangleIndicesItr())
+			{
+				if (FaceSelection[TID + FaceStart])
+				{
+					TargetTris.Add(TID);
+				}
+			}
+			if (TargetTris.IsEmpty())
+			{
+				return false;
+			}
+			FPatchBasedMeshUVGenerator PatchUVGen;
+			PatchUVGen.MergingThreshold = MergeIslandSettings.AreaDistortionThreshold;
+			PatchUVGen.MaxNormalDeviationDeg = MergeIslandSettings.MaxNormalDeviationDeg;
+			PatchUVGen.NormalSmoothingAlpha = MergeIslandSettings.NormalSmoothingAlpha;
+			PatchUVGen.NormalSmoothingRounds = MergeIslandSettings.NormalSmoothingRounds;
+			FMeshConnectedComponents InitialComponents(&Mesh);
+			InitialComponents.FindConnectedTriangles(TargetTris, [&Mesh, &UVLayer](int32 TA, int32 TB) -> bool
+				{
+					return UVLayer->AreTrianglesConnected(TA, TB);
+				}
+			);
+			TArray<TArray<int32>> MergedIslands;
+			PatchUVGen.ComputeIslandsByRegionMerging(Mesh, *UVLayer, InitialComponents, true, MergedIslands);
+			if (MergedIslands.IsEmpty())
+			{
+				// nothing merged, no need update
+				return false;
+			}
+
+			TArray<double> OriginalUVAreas;
+			TArray<FVector2f> RefUVPos;
+			OriginalUVAreas.SetNumZeroed(MergedIslands.Num());
+			RefUVPos.SetNumZeroed(MergedIslands.Num());
+			for (int32 IslandIdx = 0; IslandIdx < MergedIslands.Num(); ++IslandIdx)
+			{
+				OriginalUVAreas[IslandIdx] = FDynamicMeshUVEditor::DetermineAreaFromUVs(*UVLayer, MergedIslands[IslandIdx]);
+				// Find any valid UV as a reference position
+				for (int32 TID : MergedIslands[IslandIdx])
+				{
+					FIndex3i ElTri = UVLayer->GetTriangle(TID);
+					if (ElTri.A != INDEX_NONE)
+					{
+						RefUVPos[IslandIdx] = UVLayer->GetElement(ElTri.A);
+						break;
+					}
+				}
+			}
+
+			TArray<bool> ValidUVsComputed;
+			int32 NumSolvesFailed = PatchUVGen.ComputeUVsFromTriangleSets(Mesh, *UVLayer, MergedIslands, ValidUVsComputed);
+
+			for (int32 IslandIdx = 0; IslandIdx < MergedIslands.Num(); ++IslandIdx)
+			{
+				double NewAreaSum = FDynamicMeshUVEditor::DetermineAreaFromUVs(*UVLayer, MergedIslands[IslandIdx]);
+				if (OriginalUVAreas[IslandIdx] > 0 && NewAreaSum > 0)
+				{
+					float Scale = (float)FMath::Sqrt(OriginalUVAreas[IslandIdx] / NewAreaSum);
+					FVector2f Origin = RefUVPos[IslandIdx];
+					FDynamicMeshUVEditor::TransformTriangleSelectionUVs(*UVLayer, MergedIslands[IslandIdx], [Scale, Origin](const FVector2f& UV)
+					{
+						return (UV - Origin) * Scale + Origin;
+					});
+				}
+			}
+
+			return true;
+		},
+		true
+	);
+}
+
+
 bool BoxProjectUVs(
 	int32 TargetUVLayer,
 	FGeometryCollection& Collection,
@@ -393,45 +559,32 @@ bool BoxProjectUVs(
 	{
 		TransformIndices.Add(Collection.TransformIndex[GeomIdx]);
 	}
-	FDynamicMeshCollection CollectionMeshes(&Collection, TransformIndices, FTransform::Identity, false);
 	TSet<int32> TargetMaterials;
-	for (int32 MatID : WhichMaterials)
-	{
-		TargetMaterials.Add(MatID);
-	}
-
-	int32 NumUVLayers = Collection.NumUVLayers();
-	if (TargetUVLayer >= NumUVLayers)
-	{
-		return false;
-	}
-
-	for (int MeshIdx = 0; MeshIdx < CollectionMeshes.Meshes.Num(); MeshIdx++)
-	{
-		int32 TransformIdx = CollectionMeshes.Meshes[MeshIdx].TransformIndex;
-		FDynamicMesh3& Mesh = CollectionMeshes.Meshes[MeshIdx].AugMesh;
-
-		FMeshNormals::InitializeOverlayToPerVertexNormals(Mesh.Attributes()->PrimaryNormals(), true);
-		AugmentedDynamicMesh::InitializeOverlayToPerVertexTangents(Mesh);
-		AugmentedDynamicMesh::InitializeOverlayToPerVertexUVs(Mesh, NumUVLayers, 0);
-
-		// Apply coincident edge merge so that projection can re-determine the UV islands without having to keep existing seams
-		FMergeCoincidentMeshEdges EdgeWelder(&Mesh);
-		EdgeWelder.Apply();
-
-		TArray<int32> TargetTris;
-		for (int TID : Mesh.TriangleIndicesItr())
+	TargetMaterials.Append(WhichMaterials);
+
+	return UE::PrivateUVHelpers::ApplyDynamicMeshBasedUVProcessing(
+		TargetUVLayer, Collection,
+		TransformIndices, 
+		[&TargetMaterials, TargetUVLayer, &BoxDimensions, TargetFaces, &WhichMaterials, &OffsetUVs,
+		bOverrideBoxDimensionsWithBounds, bCenterBoxAtPivot, bUniformProjectionScale]
+		(FDynamicMesh3& Mesh, int32 TransformIndex) -> bool
 		{
-			bool bIsTextureTri = IsTriActive(
-				AugmentedDynamicMesh::GetVisibility(Mesh, TID), AugmentedDynamicMesh::GetInternal(Mesh, TID),
-				Mesh.Attributes()->GetMaterialID()->GetValue(TID), TargetMaterials, true, TargetFaces);
-			if (bIsTextureTri)
+			TArray<int32> TargetTris;
+			for (int TID : Mesh.TriangleIndicesItr())
 			{
-				TargetTris.Add(TID);
+				bool bIsTextureTri = IsTriActive(
+					AugmentedDynamicMesh::GetVisibility(Mesh, TID), AugmentedDynamicMesh::GetInternal(Mesh, TID),
+					Mesh.Attributes()->GetMaterialID()->GetValue(TID), TargetMaterials, true, TargetFaces);
+				if (bIsTextureTri)
+				{
+					TargetTris.Add(TID);
+				}
 			}
-		}
-		if (TargetTris.Num() > 0)
-		{
+			if (TargetTris.IsEmpty())
+			{
+				return false;
+			}
+
 			FDynamicMeshUVEditor UVEd(&Mesh, TargetUVLayer, false);
 			FFrame3d BoxFrame; // defaults to origin / no rotation
 			FVector3d UseBoxDimensions = bUniformProjectionScale ? FVector3d(BoxDimensions.X) : BoxDimensions;
@@ -455,17 +608,9 @@ bool BoxProjectUVs(
 			{
 				Overlap->SetElement(ElemID, OffsetUVs + Overlap->GetElement(ElemID));
 			}
-		}
-
-		Mesh.CompactInPlace();
-
-		// Re-split vertices with different UVs/normals/tangents and transfer attributes back (required because we weld edges above)
-		AugmentedDynamicMesh::SplitOverlayAttributesToPerVertex(Mesh, true, true);
-	}
-
-	CollectionMeshes.UpdateAllCollections(Collection);
-
-	return true;
+			return true;
+		},
+		true);
 }
 
 
diff --git a/Engine/Plugins/Experimental/PlanarCutPlugin/Source/PlanarCut/Public/FractureAutoUV.h b/Engine/Plugins/Experimental/PlanarCutPlugin/Source/PlanarCut/Public/FractureAutoUV.h
--- a/Engine/Plugins/Experimental/PlanarCutPlugin/Source/PlanarCut/Public/FractureAutoUV.h
+++ b/Engine/Plugins/Experimental/PlanarCutPlugin/Source/PlanarCut/Public/FractureAutoUV.h
@@ -81,8 +81,40 @@ inline bool BoxProjectUVs(
 	return BoxProjectUVs(TargetUVLayer, Collection, BoxDimensions, TargetFaces, WhichMaterials, OffsetUVs, bOverrideBoxDimensionsWithBounds, bCenterBoxAtPivot, bUniformProjectionScale);
 }
 
+struct FMergeIslandSettings
+{
+	//
+	// Parameters for island merging
+	//
+
+	// Threshold for allowed area distortion from merging islands (when we use ExpMap to compute new UVs for the merged island)
+	double AreaDistortionThreshold = 1.5;
+	// Threshold for allowed normal deviation between merge-able islands
+	double MaxNormalDeviationDeg = 45.0;
+
+	//
+	// ExpMap parameters, for computing new UVs of merged islands
+	//
+
+	int32 NormalSmoothingRounds = 0;
+	double NormalSmoothingAlpha = 0.25;
+};
+
 /**
- * Make a UV atlas of non-overlapping UV charts for a geometry collection
+ * Merge existing UV islands on the chosen faces and UV Layer, based on a normal angle threshold and a distortion threshold.
+ * The Exponential Map algorithm will be used to compute new UVs for the merged islands.
+ */
+bool PLANARCUT_API MergeUVIslands(
+	int32 TargetUVLayer,
+	FGeometryCollection& Collection,
+	FMergeIslandSettings MergeIslandSettings,
+	TArrayView<bool> FaceSelection,
+	FProgressCancel* Progress = nullptr
+);
+
+
+/**
+ * Make a UV atlas of non-overlapping UV charts for a geometry collection, using the existing UV islands
  *
  * @param TargetUVLayer		Which UV layer to update with new UV coordinates
  * @param Collection		The collection to be atlas'd
diff --git a/Engine/Plugins/Runtime/GeometryProcessing/Source/DynamicMesh/Private/Parameterization/DynamicMeshUVEditor.cpp b/Engine/Plugins/Runtime/GeometryProcessing/Source/DynamicMesh/Private/Parameterization/DynamicMeshUVEditor.cpp
--- a/Engine/Plugins/Runtime/GeometryProcessing/Source/DynamicMesh/Private/Parameterization/DynamicMeshUVEditor.cpp
+++ b/Engine/Plugins/Runtime/GeometryProcessing/Source/DynamicMesh/Private/Parameterization/DynamicMeshUVEditor.cpp
@@ -1718,7 +1718,7 @@ bool FDynamicMeshUVEditor::UDIMPack(int32 TargetTextureResolution, float GutterS
 }
 
 
-double  FDynamicMeshUVEditor::DetermineAreaFromUVs(const FDynamicMeshUVOverlay& UVOverlay, const TArray<int32>& Triangles, FAxisAlignedBox2f* BoundingBox)
+double FDynamicMeshUVEditor::DetermineAreaFromUVs(const FDynamicMeshUVOverlay& UVOverlay, const TArray<int32>& Triangles, FAxisAlignedBox2f* BoundingBox)
 {
 	if (BoundingBox)
 	{
@@ -1743,3 +1743,24 @@ double  FDynamicMeshUVEditor::DetermineAreaFromUVs(const FDynamicMeshUVOverlay&
 
 	return Area2D;
 }
+
+void FDynamicMeshUVEditor::TransformTriangleSelectionUVs(FDynamicMeshUVOverlay& UVOverlay, TConstArrayView<int32> Triangles, TFunctionRef<FVector2f(const FVector2f&)> TransformFunc)
+{
+	TSet<int32> Elements;
+	for (int32 TID : Triangles)
+	{
+		FIndex3i Tri = UVOverlay.GetTriangle(TID);
+		if (Tri.A != INDEX_NONE)
+		{
+			Elements.Add(Tri.A);
+			Elements.Add(Tri.B);
+			Elements.Add(Tri.C);
+		}
+	}
+	for (int32 EID : Elements)
+	{
+		FVector2f UV = UVOverlay.GetElement(EID);
+		UVOverlay.SetElement(EID, TransformFunc(UV));
+	}
+}
+
diff --git a/Engine/Plugins/Runtime/GeometryProcessing/Source/DynamicMesh/Private/Parameterization/PatchBasedMeshUVGenerator.cpp b/Engine/Plugins/Runtime/GeometryProcessing/Source/DynamicMesh/Private/Parameterization/PatchBasedMeshUVGenerator.cpp
--- a/Engine/Plugins/Runtime/GeometryProcessing/Source/DynamicMesh/Private/Parameterization/PatchBasedMeshUVGenerator.cpp
+++ b/Engine/Plugins/Runtime/GeometryProcessing/Source/DynamicMesh/Private/Parameterization/PatchBasedMeshUVGenerator.cpp
@@ -340,6 +340,24 @@ bool FPatchBasedMeshUVGenerator::ComputeIslandsByRegionMerging(
 	const FMeshConnectedComponents& ConnectedComponents,
 	TArray<TArray<int32>>& IslandsOut,
 	FProgressCancel* Progress)
+{
+	return ComputeIslandsByRegionMerging(
+		TargetMesh,
+		TargetUVOverlay,
+		ConnectedComponents,
+		false,
+		IslandsOut,
+		Progress);
+}
+
+
+bool FPatchBasedMeshUVGenerator::ComputeIslandsByRegionMerging(
+	FDynamicMesh3& TargetMesh,
+	FDynamicMeshUVOverlay& TargetUVOverlay,
+	const FMeshConnectedComponents& ConnectedComponents,
+	bool bOnlyOutputMergedRegions,
+	TArray<TArray<int32>>& IslandsOut,
+	FProgressCancel* Progress)
 {
 	if (Progress && Progress->Cancelled())
 	{
@@ -571,22 +589,25 @@ bool FPatchBasedMeshUVGenerator::ComputeIslandsByRegionMerging(
 
 	// accumulate remaining valid regions
 
+	auto KeepRegion = [bOnlyOutputMergedRegions, &ConnectedComponents, &RegionGraph](int32 RID) -> bool
+	{
+		int32 KeepAboveTriCount = bOnlyOutputMergedRegions ? ConnectedComponents[RID].Indices.Num() : 0;
+		return RegionGraph.GetRegionTriCount(RID) > KeepAboveTriCount;
+	};
+
 	int32 NumValidRegions = 0;
-	for (int32 rid : RemainingRegions)
+	for (int32 RID : RemainingRegions)
 	{
-		if (RegionGraph.IsRegion(rid) && RegionGraph.GetRegionTris(rid).Num() > 0)
-		{
-			NumValidRegions++;
-		}
+		NumValidRegions += int32(KeepRegion(RID));
 	}
 
 	IslandsOut.Reserve(NumValidRegions);
 
-	for (int32 rid : RemainingRegions)
+	for (int32 RID : RemainingRegions)
 	{
-		if (RegionGraph.IsRegion(rid))
+		if (KeepRegion(RID))
 		{
-			IslandsOut.Add(RegionGraph.MoveRegionTris(rid));
+			IslandsOut.Add(RegionGraph.MoveRegionTris(RID));
 		}
 		if (Progress && Progress->Cancelled())
 		{
diff --git a/Engine/Plugins/Runtime/GeometryProcessing/Source/DynamicMesh/Public/Parameterization/DynamicMeshUVEditor.h b/Engine/Plugins/Runtime/GeometryProcessing/Source/DynamicMesh/Public/Parameterization/DynamicMeshUVEditor.h
--- a/Engine/Plugins/Runtime/GeometryProcessing/Source/DynamicMesh/Public/Parameterization/DynamicMeshUVEditor.h
+++ b/Engine/Plugins/Runtime/GeometryProcessing/Source/DynamicMesh/Public/Parameterization/DynamicMeshUVEditor.h
@@ -348,6 +348,11 @@ public:
 	*/
 	static double DetermineAreaFromUVs(const FDynamicMeshUVOverlay& UVOverlay, const TArray<int32>& Triangles, FAxisAlignedBox2f* BoundingBox = nullptr);
 
+	/**
+	 * Transform UVs of the specified triangles
+	 */
+	static void TransformTriangleSelectionUVs(FDynamicMeshUVOverlay& UVOverlay, TConstArrayView<int32> Triangles, TFunctionRef<FVector2f(const FVector2f&)> TransformFunc);
+
 private: 
 
 	struct FSetUVsFromConformalOptions
@@ -365,7 +370,6 @@ private:
 
 	// Helper to split bowties on element IDs
 	void SplitBowtiesOnUVElements(TArray<int32>& UVElementIDs, bool bAddNewElementsToInputArray = true);
-	
 };
 
 
diff --git a/Engine/Plugins/Runtime/GeometryProcessing/Source/DynamicMesh/Public/Parameterization/PatchBasedMeshUVGenerator.h b/Engine/Plugins/Runtime/GeometryProcessing/Source/DynamicMesh/Public/Parameterization/PatchBasedMeshUVGenerator.h
--- a/Engine/Plugins/Runtime/GeometryProcessing/Source/DynamicMesh/Public/Parameterization/PatchBasedMeshUVGenerator.h
+++ b/Engine/Plugins/Runtime/GeometryProcessing/Source/DynamicMesh/Public/Parameterization/PatchBasedMeshUVGenerator.h
@@ -87,7 +87,9 @@ public:
 
 	/**
 	 * Incrementally combine existing mesh patches into larger patches that meet the various criteria defined
-	 * by the parameters above
+	 * by the parameters above.
+	 * Note: This overload always outputs all islands (i.e., equivalent to calling the other overload with bOnlyOutputMergedRegions=false)
+	 * 
 	 * @param InitialComponents the initial mesh patches, ie triangle lists
 	 * @param IslandsOut the set of output mesh patches / triangle lists
 	 * @return false if catastrophic failure ocurred (ie IslandsOut is invalid) or if Progress was cancelled
@@ -99,6 +101,23 @@ public:
 		TArray<TArray<int32>>& IslandsOut,
 		FProgressCancel* Progress = nullptr);
 
+	/**
+	 * Incrementally combine existing mesh patches into larger patches that meet the various criteria defined
+	 * by the parameters above
+	 * 
+	 * @param InitialComponents the initial mesh patches, ie triangle lists
+	 * @param bOnlyOutputMergedRegions whether to restrict IslandsOut to just the regions where components were merged, skipping unchanged islands
+	 * @param IslandsOut the set of output mesh patches / triangle lists
+	 * @return false if catastrophic failure ocurred (ie IslandsOut is invalid) or if Progress was cancelled
+	 */
+	bool ComputeIslandsByRegionMerging(
+		FDynamicMesh3& TargetMesh,
+		FDynamicMeshUVOverlay& TargetUVOverlay,
+		const FMeshConnectedComponents& InitialComponents,
+		bool bOnlyOutputMergedRegions,
+		TArray<TArray<int32>>& IslandsOut,
+		FProgressCancel* Progress = nullptr);
+
 	/**
 	 * Calculate UVs for the input mesh patches of TargetMesh and store in TargetUVOVerlay
 	 * @param TriangleSets the set of desired UV islands, defined as triangle lists
