make fracture mesh conversion more robust to overlays with unset triangles

diff --git a/Engine/Plugins/Experimental/PlanarCutPlugin/Source/PlanarCut/Private/GeometryMeshConversion.cpp b/Engine/Plugins/Experimental/PlanarCutPlugin/Source/PlanarCut/Private/GeometryMeshConversion.cpp
--- a/Engine/Plugins/Experimental/PlanarCutPlugin/Source/PlanarCut/Private/GeometryMeshConversion.cpp
+++ b/Engine/Plugins/Experimental/PlanarCutPlugin/Source/PlanarCut/Private/GeometryMeshConversion.cpp
@@ -326,6 +326,12 @@ namespace AugmentedDynamicMesh
 			return Overlay->GetTriangle(TID)[SubIdx];
 		};
 
+		// get overlay element, returning a zero vector for unset elements
+		auto GetElement = [&Overlay](int32 EID)
+		{
+			return EID > -1 ? Overlay->GetElement(EID) : VecType(EForceInit::ForceInitToZero);
+		};
+
 		for (int32 VID = 0, OrigMax = Mesh.MaxVertexID(); VID < OrigMax; VID++)
 		{
 			if (!Mesh.IsVertex(VID))
@@ -341,7 +347,7 @@ namespace AugmentedDynamicMesh
 				checkSlow(TriStartIdx < ContigTris.Num());
 				int32 TID0 = ContigTris[TriStartIdx];
 				int32 EID0 = GetEID(VID, TID0);
-				VecType LastElData = Overlay->GetElement(EID0);
+				VecType LastElData = GetElement(EID0);
 				int32 LastEID = EID0;
 				bool bPastEID0 = false;
 				ToSplitTris.Reset();
@@ -351,7 +357,7 @@ namespace AugmentedDynamicMesh
 					int32 EID = GetEID(VID, TID);
 					VecType ElData;
 					if (EID != LastEID &&
-						!VectorUtil::EpsilonEqual(ElData = Overlay->GetElement(EID), LastElData, FMathf::ZeroTolerance))
+						!VectorUtil::EpsilonEqual(ElData = GetElement(EID), LastElData, FMathf::ZeroTolerance))
 					{
 						if (ToSplitTris.Num() > 0)
 						{
