add area contact option for small geo removal

diff --git a/Engine/Plugins/Experimental/ChaosEditor/Source/FractureEditor/Private/FractureToolFixTinyGeo.cpp b/Engine/Plugins/Experimental/ChaosEditor/Source/FractureEditor/Private/FractureToolFixTinyGeo.cpp
--- a/Engine/Plugins/Experimental/ChaosEditor/Source/FractureEditor/Private/FractureToolFixTinyGeo.cpp
+++ b/Engine/Plugins/Experimental/ChaosEditor/Source/FractureEditor/Private/FractureToolFixTinyGeo.cpp
@@ -296,6 +296,10 @@ void UFractureToolFixTinyGeo::Execute(TWeakPtr<FFractureEditorModeToolkit> InToo
 			{
 				SelectionMethod = UE::PlanarCut::ENeighborSelectionMethod::NearestCenter;
 			}
+			else if (TinyGeoSettings->NeighborSelection == ENeighborSelectionMethod::LargestContactArea)
+			{
+				SelectionMethod = UE::PlanarCut::ENeighborSelectionMethod::LargestContactArea;
+			}
 
 			if (TinyGeoSettings->MergeType == EMergeType::MergeGeometry)
 			{
diff --git a/Engine/Plugins/Experimental/ChaosEditor/Source/FractureEditor/Private/FractureToolFixTinyGeo.h b/Engine/Plugins/Experimental/ChaosEditor/Source/FractureEditor/Private/FractureToolFixTinyGeo.h
--- a/Engine/Plugins/Experimental/ChaosEditor/Source/FractureEditor/Private/FractureToolFixTinyGeo.h
+++ b/Engine/Plugins/Experimental/ChaosEditor/Source/FractureEditor/Private/FractureToolFixTinyGeo.h
@@ -24,7 +24,9 @@ enum class ENeighborSelectionMethod
 	// Merge to the neighbor with the largest volume
 	LargestNeighbor,
 	// Merge to the neighbor with the closest center
-	NearestCenter
+	NearestCenter,
+	// Merge to the neighbor with the largest contact surface area
+	LargestContactArea
 };
 
 UENUM()
diff --git a/Engine/Plugins/Experimental/Fracture/Source/FractureEngine/Private/FractureEngineUtility.cpp b/Engine/Plugins/Experimental/Fracture/Source/FractureEngine/Private/FractureEngineUtility.cpp
--- a/Engine/Plugins/Experimental/Fracture/Source/FractureEngine/Private/FractureEngineUtility.cpp
+++ b/Engine/Plugins/Experimental/Fracture/Source/FractureEngine/Private/FractureEngineUtility.cpp
@@ -264,6 +264,10 @@ void FFractureEngineUtility::FixTinyGeo(FManagedArrayCollection& InOutCollection
 			{
 				SelectionMethod = UE::PlanarCut::ENeighborSelectionMethod::NearestCenter;
 			}
+			else if (InNeighborSelection == EFixTinyGeoNeighborSelectionMethod::LargestContactArea)
+			{
+				SelectionMethod = UE::PlanarCut::ENeighborSelectionMethod::LargestContactArea;
+			}
 
 			// Make sure we have non-stale proximity data if we will use it
 			if (bUseCollectionProximity)
diff --git a/Engine/Plugins/Experimental/Fracture/Source/FractureEngine/Public/FractureEngineUtility.h b/Engine/Plugins/Experimental/Fracture/Source/FractureEngine/Public/FractureEngineUtility.h
--- a/Engine/Plugins/Experimental/Fracture/Source/FractureEngine/Public/FractureEngineUtility.h
+++ b/Engine/Plugins/Experimental/Fracture/Source/FractureEngine/Public/FractureEngineUtility.h
@@ -27,6 +27,8 @@ enum class EFixTinyGeoNeighborSelectionMethod : uint8
 	LargestNeighbor UMETA(DisplayName = "Largest Neighbor"),
 	// Merge to the neighbor with the closest center
 	NearestCenter UMETA(DisplayName = "Nearest Center"),
+	// Merge to the neighbor with the largest contact surface area
+	LargestContactArea
 };
 
 UENUM(BlueprintType)
diff --git a/Engine/Plugins/Experimental/PlanarCutPlugin/Source/PlanarCut/Private/PlanarCut.cpp b/Engine/Plugins/Experimental/PlanarCutPlugin/Source/PlanarCut/Private/PlanarCut.cpp
--- a/Engine/Plugins/Experimental/PlanarCutPlugin/Source/PlanarCut/Private/PlanarCut.cpp
+++ b/Engine/Plugins/Experimental/PlanarCutPlugin/Source/PlanarCut/Private/PlanarCut.cpp
@@ -4,6 +4,8 @@
 #include "PlanarCutPlugin.h"
 
 #include "Spatial/SparseDynamicOctree3.h"
+#include "Spatial/MeshAABBTree3.h"
+#include "MeshQueries.h"
 #include "Util/IndexUtil.h"
 
 #include "GeometryCollection/GeometryCollectionAlgo.h"
@@ -1231,7 +1233,7 @@ void FindBoneVolumes(
 		int32 TransformIdx = TransformIndicesArray[Idx];
 		if (bIncludeClusters || Collection.IsRigid(TransformIdx))
 		{
-			double UnscaledVolume = (double)VolumesAttrib[Idx];
+			double UnscaledVolume = (double)VolumesAttrib[TransformIdx];
 			OutVolumes[Idx] = UnscaledVolume * VolumeScale;
 		}
 		else
@@ -1292,6 +1294,73 @@ void FilterBonesByVolume(
 	}
 }
 
+// mesh adapter for a single geometry in a geometry collection
+// Note: implicitly remaps vertices and faces to start at 0, to avoid wasting space in the Mesh AABB tree use case
+struct FGeometryCollectionGeometryMeshAdapter
+{
+	TConstArrayView<FIntVector> FaceView;
+	TConstArrayView<FVector3f> VertexView;
+	int32 FaceStart, FaceNum, VertexStart, VertexNum;
+
+	FGeometryCollectionGeometryMeshAdapter(const FGeometryCollection& Collection, int32 GeometryIdx)
+	{
+		FaceStart = Collection.FaceStart[GeometryIdx];
+		FaceNum = Collection.FaceCount[GeometryIdx];
+		VertexStart = Collection.VertexStart[GeometryIdx];
+		VertexNum = Collection.VertexCount[GeometryIdx];
+		FaceView = Collection.Indices.GetConstArray();
+		VertexView = Collection.Vertex.GetConstArray();
+	}
+
+	bool IsTriangle(int32 TID) const
+	{
+		return TID >= 0 && TID < FaceNum;
+	}
+	bool IsVertex(int32 VID) const
+	{
+		return VID >= 0 && VID < VertexNum;
+	}
+	int32 MaxTriangleID() const
+	{
+		return FaceNum;
+	}
+	int32 MaxVertexID() const
+	{
+		return VertexNum;
+	}
+	int32 TriangleCount() const
+	{
+		return FaceNum;
+	}
+	int32 VertexCount() const
+	{
+		return VertexNum;
+	}
+	// note: GetChangeStamp is required by mesh aabb tree interface but not needed/useful otherwise
+	uint64 GetChangeStamp() const
+	{
+		return 0;
+	}
+	FIndex3i GetTriangle(int32 TID) const
+	{
+		FIntVector Tri = FaceView[TID + FaceStart];
+		return FIndex3i(Tri.X - VertexStart, Tri.Y - VertexStart, Tri.Z - VertexStart);
+	}
+	FVector3f GetVertex(int32 VID) const
+	{
+		return VertexView[VID + VertexStart];
+	}
+
+	template<typename VertexType>
+	inline void GetTriVertices(int TID, VertexType& V0, VertexType& V1, VertexType& V2) const
+	{
+		FIntVector Tri = FaceView[TID + FaceStart];
+		V0 = (VertexType)VertexView[Tri.X];
+		V1 = (VertexType)VertexView[Tri.Y];
+		V2 = (VertexType)VertexView[Tri.Z];
+	}
+};
+
 
 void FindSmallBones(
 	const FGeometryCollection& Collection,
@@ -1311,6 +1380,59 @@ void FindSmallBones(
 }
 
 
+static TMap<UE::Geometry::FIndex2i, float> ComputeGeometryVolumeContactWeightsFromProximity(
+	const FGeometryCollection& Collection,
+	const TArray<TSet<int32>>& GeomProximity,
+	const TSet<int32>& GeomIdxSelection,
+	const TMap<int32, double>* OptionalGeomToVol = nullptr,
+	float DistanceThreshold = .01, float SamplesPerUnitArea = 10, int32 MaxSamplesPerTri = 100)
+{
+	TMap<UE::Geometry::FIndex2i, float> ContactWeights;
+	TMultiMap<int32, int32> NeedConnections;
+	for (int32 GeomIdx : GeomIdxSelection)
+	{
+		double GeomVol = OptionalGeomToVol ? (*OptionalGeomToVol)[GeomIdx] : 0.;
+		for (int32 GeomNbrIdx : GeomProximity[GeomIdx])
+		{
+			// filter out contacts that we'll see twice
+			if (GeomNbrIdx < GeomIdx && GeomIdxSelection.Contains(GeomNbrIdx))
+			{
+				continue;
+			}
+
+			// process in a consistent order. if we have volumes calculated, favor small->large,
+			// since that may be a bit faster (typically fewer trianges/less surface area to sample)
+			double NbrVol = OptionalGeomToVol ? (*OptionalGeomToVol)[GeomNbrIdx] : 0.;
+			int32 GeomB = GeomIdx, GeomA = GeomNbrIdx;
+			if (NbrVol < GeomVol || (NbrVol == GeomVol && GeomB < GeomA))
+			{
+				Swap(GeomA, GeomB);
+			}
+			NeedConnections.Add(GeomA, GeomB);
+		}
+	}
+	TArray<int32> Keys;
+	NeedConnections.GetKeys(Keys);
+	for (int32 KeyGeom : Keys)
+	{
+		FGeometryCollectionGeometryMeshAdapter RefMesh(Collection, KeyGeom);
+		TMeshAABBTree3<FGeometryCollectionGeometryMeshAdapter> RefSpatial(&RefMesh, true);
+		FTransform3f RefTransform = (FTransform3f)GeometryCollectionAlgo::GlobalMatrix(Collection.Transform, Collection.Parent, Collection.TransformIndex[KeyGeom]);
+		for (auto KeyItr = NeedConnections.CreateConstKeyIterator(KeyGeom); KeyItr; ++KeyItr)
+		{
+			int32 NbrGeom = KeyItr.Value();
+			FTransform3f NbrTransform = (FTransform3f)GeometryCollectionAlgo::GlobalMatrix(Collection.Transform, Collection.Parent, Collection.TransformIndex[NbrGeom]);
+			FGeometryCollectionGeometryMeshAdapter NbrMesh(Collection, NbrGeom);
+			float ContactArea = UE::Geometry::TMeshQueries<FGeometryCollectionGeometryMeshAdapter>::ApproximateContactArea<TMeshAABBTree3<FGeometryCollectionGeometryMeshAdapter>, float>
+				(NbrMesh, NbrTransform, RefSpatial, RefTransform, DistanceThreshold, SamplesPerUnitArea, MaxSamplesPerTri);
+			FIndex2i OutputKey(KeyGeom, NbrGeom);
+			OutputKey.Sort();
+			ContactWeights.Add(OutputKey, ContactArea);
+		}
+	}
+	return ContactWeights;
+}
+
 int32 MergeBones(
 	FGeometryCollection& Collection,
 	const TArrayView<const int32>& TransformIndicesView,
@@ -1475,7 +1597,7 @@ int32 MergeBones(
 		{
 			CanMerge.Add(GeomIdx);
 			GeomToVol.Add(GeomIdx, Volumes[Idx]);
-			GeomToParent.Add(GeomIdx, Collection.Parent[Idx]);
+			GeomToParent.Add(GeomIdx, Collection.Parent[TransformIdx]);
 		}
 	}
 	for (int32 TransformIdx : SmallTransformIndices)
@@ -1491,6 +1613,13 @@ int32 MergeBones(
 		}
 	}
 
+	TMap<UE::Geometry::FIndex2i, float> ConnectionWeights;
+	if (NeighborSelectionMethod == ENeighborSelectionMethod::LargestContactArea)
+	{
+		ConnectionWeights = ComputeGeometryVolumeContactWeightsFromProximity(Collection, *UseProximity, TooSmalls, &GeomToVol);
+	}
+
+
 	for (int32 SmallIdx : TooSmalls)
 	{
 		int32* SmallRemoveGroupIdx = GeomIdxToRemoveGroupIdx.Find(SmallIdx);
@@ -1518,10 +1647,17 @@ int32 MergeBones(
 				{
 					Score = GeomToVol[NbrIdx];
 				}
-				else // Nearest center
+				else if (NeighborSelectionMethod == UE::PlanarCut::ENeighborSelectionMethod::NearestCenter)
 				{
 					Score = 1.0 / (DOUBLE_SMALL_NUMBER + DistanceSquared(GetCenter(NbrIdx), GetCenter(SmallIdx)));
 				}
+				else // use pre-computed connection weight map
+				{
+					FIndex2i Key(SmallIdx, NbrIdx);
+					Key.Sort();
+					float* FoundWt = ConnectionWeights.Find(Key);
+					Score = FoundWt ? *FoundWt : 0.f;
+				}
 				if (Score > BestScore)
 				{
 					BestScore = Score;
@@ -1706,6 +1842,22 @@ namespace
 		}
 		return Combined;
 	}
+	
+	static void GetRigidGeometrySelectionFromBoneSelection(const FGeometryCollection& Collection, int32 BoneIdx, TSet<int32>& OutGeometryIndices)
+	{
+		int32 GeoIdx = Collection.TransformToGeometryIndex[BoneIdx];
+		if (Collection.IsRigid(BoneIdx) && GeoIdx != INDEX_NONE)
+		{
+			OutGeometryIndices.Add(GeoIdx);
+		}
+		else if (Collection.IsClustered(BoneIdx))
+		{
+			for (int32 Child : Collection.Children[BoneIdx])
+			{
+				GetRigidGeometrySelectionFromBoneSelection(Collection, Child, OutGeometryIndices);
+			}
+		}
+	}
 }
 
 void MergeClusters(
@@ -1814,6 +1966,17 @@ void MergeClusters(
 			}
 		}
 	};
+	
+	TSet<int32> ToProcessGeometry;
+	for (int32 BoneIdx : SmallTransformIndices)
+	{
+		GetRigidGeometrySelectionFromBoneSelection(Collection, BoneIdx, ToProcessGeometry);
+	}
+	TMap<UE::Geometry::FIndex2i, float> ConnectionWeights;
+	if (NeighborSelectionMethod == ENeighborSelectionMethod::LargestContactArea)
+	{
+		ConnectionWeights = ComputeGeometryVolumeContactWeightsFromProximity(Collection, *UseProximity, ToProcessGeometry);
+	}
 
 	TSet<int32> ReconsiderIfSmall, DoNotReconsider; // Sets to track nodes we've already merged, and may not need to re-merge
 	TSet<int32> Neighbors; // Set to be re-used in the below loop, tracking neighbors to consider
@@ -1875,7 +2038,7 @@ void MergeClusters(
 					}
 				}
 			}
-			else // ENeighborSelectionMethod::LargestNeighbor
+			else if (NeighborSelectionMethod == ENeighborSelectionMethod::LargestNeighbor)
 			{
 				double Volume = GetClusterVolumeFromRigidsHelper(Collection, Volumes, MergeIdx);
 				double LargestVolume = -1;
@@ -1889,6 +2052,42 @@ void MergeClusters(
 					}
 				}
 			}
+			else // ENeighborSelectionMethod::LargestContactArea
+			{
+				TSet<int32> MergeGeo, NbrGeo;
+				GetRigidGeometrySelectionFromBoneSelection(Collection, MergeIdx, MergeGeo);
+				double LargestContact = -1;
+				for (int32 NeighborIdx : Neighbors)
+				{
+					double ContactArea = 0.;
+					NbrGeo.Reset();
+					GetRigidGeometrySelectionFromBoneSelection(Collection, NeighborIdx, NbrGeo);
+					// Look for proximity-neighbors of the merge bone's geometry
+					// that belong to the current neighbor bone's geometry,
+					// and add the corresponding contact area
+					for (int32 MergeGeoIdx : MergeGeo)
+					{
+						for (int32 MergeGeoNbrIdx : (*UseProximity)[MergeGeoIdx])
+						{
+							if (NbrGeo.Contains(MergeGeoNbrIdx))
+							{
+								FIndex2i Key(MergeGeoIdx, MergeGeoNbrIdx);
+								Key.Sort();
+								float* FoundArea = ConnectionWeights.Find(Key);
+								if (FoundArea)
+								{
+									ContactArea += double(*FoundArea);
+								}
+							}
+						}
+					}
+					if (BestNeighbor == INDEX_NONE || ContactArea > LargestContact)
+					{
+						BestNeighbor = NeighborIdx;
+						LargestContact = ContactArea;
+					}
+				}
+			}
 		}
 		if (BestNeighbor == INDEX_NONE)
 		{
diff --git a/Engine/Plugins/Experimental/PlanarCutPlugin/Source/PlanarCut/Public/PlanarCut.h b/Engine/Plugins/Experimental/PlanarCutPlugin/Source/PlanarCut/Public/PlanarCut.h
--- a/Engine/Plugins/Experimental/PlanarCutPlugin/Source/PlanarCut/Public/PlanarCut.h
+++ b/Engine/Plugins/Experimental/PlanarCutPlugin/Source/PlanarCut/Public/PlanarCut.h
@@ -386,7 +386,8 @@ namespace UE
 		enum ENeighborSelectionMethod
 		{
 			LargestNeighbor,
-			NearestCenter
+			NearestCenter,
+			LargestContactArea
 		};
 	}
 }
