Fix curvature and AO bakes not working correctly w/ gutter padding (because their per-islands were using the expanded gutters as well and overwriting each other)

#Synced-CL 50388552

diff --git a/Engine/Plugins/Experimental/PlanarCutPlugin/Source/PlanarCut/Private/FractureAutoUV.cpp b/Engine/Plugins/Experimental/PlanarCutPlugin/Source/PlanarCut/Private/FractureAutoUV.cpp
--- a/Engine/Plugins/Experimental/PlanarCutPlugin/Source/PlanarCut/Private/FractureAutoUV.cpp
+++ b/Engine/Plugins/Experimental/PlanarCutPlugin/Source/PlanarCut/Private/FractureAutoUV.cpp
@@ -938,6 +938,11 @@ bool TextureSpecifiedFaces(
 		return false;
 	}
 
+	// Gutter for separate bakes used for curvature / AO / etc -- we shouldn't need the full gutter step for these,
+	//  because we do a post-process later to fill the gutter for the whole image.
+	// But we still need a positive gutter size for the bakes to work, so we use a size of 1.
+	const int32 PartialBakeGutterSize = 1;
+
 	TArray64<float> CurvatureValues; // buffer to fill with curvature values
 	if (CurvatureIdx > -1)
 	{
@@ -987,7 +992,7 @@ bool TextureSpecifiedFaces(
 			FDynamicMeshAABBTree3 SolidSpatial(&SolidMesh);
 
 			FMeshImageBakingCache BakeCache;
-			BakeCache.SetGutterSize(GutterSize);
+			BakeCache.SetGutterSize(PartialBakeGutterSize);
 			BakeCache.SetDetailMesh(&SolidMesh, &SolidSpatial);
 			BakeCache.SetBakeTargetMesh(&Mesh);
 			BakeCache.SetDimensions(Dimensions);
@@ -1037,7 +1042,7 @@ bool TextureSpecifiedFaces(
 			FDynamicMeshAABBTree3 Spatial(&Mesh);
 
 			FMeshImageBakingCache BakeCache;
-			BakeCache.SetGutterSize(GutterSize);
+			BakeCache.SetGutterSize(PartialBakeGutterSize);
 			BakeCache.SetDetailMesh(&Mesh, &Spatial);
 			BakeCache.SetBakeTargetMesh(&Mesh);
 			BakeCache.SetDimensions(Dimensions);
