Change the geometrycore MeshSpatialSort algorithm to not hold a TArrayView of meshes, and instead pass a mesh accessor to the Compute() method. Support taking arrays of (unique)pointers-to-mesh instead of arrays-of-mesh, so the API works better with the new/safer version of FDynamicMeshEditor::SplitMesh.

Update PlanarCut's usage of MeshSpatialSort and SplitMesh to the new/safer pointer versions.

#Synced-CL 50490755

diff --git a/Engine/Plugins/Experimental/PlanarCutPlugin/Source/PlanarCut/Private/GeometryMeshConversion.cpp b/Engine/Plugins/Experimental/PlanarCutPlugin/Source/PlanarCut/Private/GeometryMeshConversion.cpp
--- a/Engine/Plugins/Experimental/PlanarCutPlugin/Source/PlanarCut/Private/GeometryMeshConversion.cpp
+++ b/Engine/Plugins/Experimental/PlanarCutPlugin/Source/PlanarCut/Private/GeometryMeshConversion.cpp
@@ -2874,15 +2874,17 @@ int32 FDynamicMeshCollection::CutWithMultiplePlanes(
 				TArray<int32, TInlineAllocator<4>> ResultIndices = { ToCutIdx, NewIdx };
 				// corresponding parent indices for each result mesh
 				TArray<int32, TInlineAllocator<4>> ParentIndices = { 0, 1 };
-				TArray<FDynamicMesh3> SplitMeshes;
+				TArray<TUniquePtr<FDynamicMesh3>> SplitMeshes;
 				for (int UnsplitIdx = 0; UnsplitIdx < 2; UnsplitIdx++)
 				{
 					if (bSplitIslands && SplitIslands(ToCut[ResultIndices[UnsplitIdx]]->AugMesh, SplitMeshes))
 					{
-						ToCut[ResultIndices[UnsplitIdx]]->SetMesh(SplitMeshes[0]);
+						// TODO: Add a SetMesh(FDynamicMesh3&&) method so we can MoveTemp here
+						ToCut[ResultIndices[UnsplitIdx]]->SetMesh(*SplitMeshes[0]);
 						for (int32 Idx = 1; Idx < SplitMeshes.Num(); Idx++)
 						{
-							const FDynamicMesh3& Mesh = SplitMeshes[Idx];
+							const FDynamicMesh3& Mesh = *SplitMeshes[Idx];
+							// TODO: Add a FMeshData constructor to let us MoveTemp the mesh here
 							ResultIndices.Add(ToCut.Add(MakeUnique<FMeshData>(Mesh, TransformIndex, FromCollection)));
 							ParentIndices.Add(UnsplitIdx);
 						}
@@ -2977,12 +2979,12 @@ int32 FDynamicMeshCollection::SplitAllIslands(FGeometryCollection* Collection, d
 		int32 SrcGeometryIdx = Collection->TransformToGeometryIndex[Surface.TransformIndex];
 		
 		int32 CreatedGeometryIdx = -1;
-		TArray<FDynamicMesh3> Islands;
+		TArray<TUniquePtr<FDynamicMesh3>> Islands;
 		if (SplitIslands(Surface.AugMesh, Islands))
 		{
 			for (int32 i = 0; i < Islands.Num(); i++)
 			{
-				FDynamicMesh3& Island = Islands[i];
+				FDynamicMesh3& Island = *Islands[i];
 				FString BoneName = GetBoneName(*Collection, Surface.TransformIndex, i);
 				constexpr int32 InternalMaterialID = 0; // Note: there won't be new internal faces (to assign materials) from a split, so this value won't affect the output here
 				CreatedGeometryIdx = AppendToCollection(Surface.FromCollection, Island, CollisionSampleSpacing, Surface.TransformIndex, BoneName, *Collection, InternalMaterialID);
@@ -3105,12 +3107,12 @@ int32 FDynamicMeshCollection::CutWithCellMeshes(const FInternalSurfaceMaterials&
 						}
 					}
 					int32 CreatedGeometryIdx = -1;
-					TArray<FDynamicMesh3> Islands;
+					TArray<TUniquePtr<FDynamicMesh3>> Islands;
 					if (bSplitIslands && SplitIslands(AugBoolResult, Islands))
 					{
 						for (int32 i = 0; i < Islands.Num(); i++)
 						{
-							FDynamicMesh3& Island = Islands[i];
+							FDynamicMesh3& Island = *Islands[i];
 							FString BoneName = GetBoneName(*Collection, Surface.TransformIndex, SubPartIndex++);
 							CreatedGeometryIdx = AppendToCollection(Surface.FromCollection, Island, CollisionSampleSpacing, Surface.TransformIndex, BoneName, *Collection, InternalMaterialID);
 							CellToGeometry.Add(CellIdx, CreatedGeometryIdx);
@@ -3224,8 +3226,10 @@ bool FDynamicMeshCollection::IsNeighboring(
 }
 
 // Split mesh into connected components, including implicit connections by co-located vertices
-bool FDynamicMeshCollection::SplitIslands(FDynamicMesh3& Source, TArray<FDynamicMesh3>& SeparatedMeshes, double SnapDistance)
+bool FDynamicMeshCollection::SplitIslands(FDynamicMesh3& Source, TArray<TUniquePtr<FDynamicMesh3>>& SeparatedMeshes, double SnapDistance)
 {
+	SeparatedMeshes.Reset();
+
 	TPointHashGrid3d<int> VertHash(SnapDistance * 10, -1);
 	FDisjointSet VertComponents(Source.MaxVertexID());
 	// Add Source vertices to hash & disjoint sets
@@ -3245,7 +3249,7 @@ bool FDynamicMeshCollection::SplitIslands(FDynamicMesh3& Source, TArray<FDynamic
 		VertComponents.Union(Tri.B, Tri.C);
 	}
 
-	bool bWasSplit = FDynamicMeshEditor::SplitMesh(&Source, SeparatedMeshes, [&Source, &VertComponents](int TID)
+	bool bWasSplit = FDynamicMeshEditor::SplitMesh(&Source, SeparatedMeshes, false, [&Source, &VertComponents](int TID)
 		{
 			return (int)VertComponents.Find(Source.GetTriangle(TID).A);
 		});
@@ -3253,19 +3257,19 @@ bool FDynamicMeshCollection::SplitIslands(FDynamicMesh3& Source, TArray<FDynamic
 	if (bWasSplit)
 	{
 		// disconnected components that are contained inside other components need to be re-merged
-		TMeshSpatialSort<FDynamicMesh3> SpatialSort(SeparatedMeshes);
+		TMeshSpatialSort<FDynamicMesh3> SpatialSort;
 		SpatialSort.NestingMethod = TMeshSpatialSort<FDynamicMesh3>::ENestingMethod::InLargestParent;
 		SpatialSort.bOnlyNestNegativeVolumes = false;
 		SpatialSort.bOnlyParentPostiveVolumes = true;
-		SpatialSort.Compute();
+		SpatialSort.Compute(SeparatedMeshes);
 		TArray<bool> KeepMeshes; KeepMeshes.Init(true, SeparatedMeshes.Num());
 		for (TMeshSpatialSort<FDynamicMesh3>::FMeshNesting& Nest : SpatialSort.Nests)
 		{
-			FDynamicMeshEditor Editor(&SeparatedMeshes[Nest.OuterIndex]);
+			FDynamicMeshEditor Editor(SeparatedMeshes[Nest.OuterIndex].Get());
 			FMeshIndexMappings Mappings;
 			for (int Inner : Nest.InnerIndices)
 			{
-				Editor.AppendMesh(&SeparatedMeshes[Inner], Mappings);
+				Editor.AppendMesh(SeparatedMeshes[Inner].Get(), Mappings);
 				KeepMeshes[Inner] = false;
 			}
 		}
diff --git a/Engine/Plugins/Experimental/PlanarCutPlugin/Source/PlanarCut/Private/GeometryMeshConversion.h b/Engine/Plugins/Experimental/PlanarCutPlugin/Source/PlanarCut/Private/GeometryMeshConversion.h
--- a/Engine/Plugins/Experimental/PlanarCutPlugin/Source/PlanarCut/Private/GeometryMeshConversion.h
+++ b/Engine/Plugins/Experimental/PlanarCutPlugin/Source/PlanarCut/Private/GeometryMeshConversion.h
@@ -304,7 +304,7 @@ struct FDynamicMeshCollection
 	}
 
 	// Split mesh into connected components, including implicit connections by co-located vertices
-	bool SplitIslands(UE::Geometry::FDynamicMesh3& Source, TArray<UE::Geometry::FDynamicMesh3>& SeparatedMeshes, double ToleranceDistance = 1e-3);
+	bool SplitIslands(UE::Geometry::FDynamicMesh3& Source, TArray<TUniquePtr<UE::Geometry::FDynamicMesh3>>& SeparatedMeshes, double ToleranceDistance = 1e-3);
 
 	FString GetBoneName(FGeometryCollection& Output, int TransformParent, int SubPartIndex)
 	{
diff --git a/Engine/Source/Runtime/GeometryCore/Public/Spatial/MeshSpatialSort.h b/Engine/Source/Runtime/GeometryCore/Public/Spatial/MeshSpatialSort.h
--- a/Engine/Source/Runtime/GeometryCore/Public/Spatial/MeshSpatialSort.h
+++ b/Engine/Source/Runtime/GeometryCore/Public/Spatial/MeshSpatialSort.h
@@ -38,6 +38,7 @@ public:
 	/**
 	 * Meshes to be spatially sorted
 	 */
+	UE_DEPRECATED(5.8, "This is no longer stored on the class; instead an accessor is passed in per-call")
 	TArrayView<const TriangleMeshType> InputMeshes;
 
 	/**
@@ -95,24 +96,44 @@ public:
 
 	TMeshSpatialSort() {}
 
+	UE_DEPRECATED(5.8, "Pass meshes and bounds to the Compute() method directly")
 	TMeshSpatialSort(TArrayView<const TriangleMeshType> InputMeshes, TArrayView<const FAxisAlignedBox3d> MeshBoundsIn = TArrayView<const FAxisAlignedBox3d>())
-		: InputMeshes(InputMeshes)
 	{
-		if (MeshBoundsIn.Num() == InputMeshes.Num())
-		{
-			MeshBounds = MeshBoundsIn;
-		}
-		else
-		{
-			InitBounds();
-		}
+		ensure(false);
 	}
 
+	UE_DEPRECATED(5.8, "Use a version of Compute() that takes the input meshes directly")
 	void Compute()
 	{
-		if (!ensure(MeshBounds.Num() == InputMeshes.Num()))
+		ensure(false);
+	}
+
+	void Compute(TArrayView<const TriangleMeshType> InInputMeshes, TArrayView<const FAxisAlignedBox3d> MeshBoundsIn = TArrayView<const FAxisAlignedBox3d>())
+	{
+		Compute([&InInputMeshes](int32 Idx) { return &InInputMeshes[Idx]; }, InInputMeshes.Num(), MeshBoundsIn);
+	}
+
+	void Compute(TArrayView<const TriangleMeshType*> InInputMeshes, TArrayView<const FAxisAlignedBox3d> MeshBoundsIn = TArrayView<const FAxisAlignedBox3d>())
+	{
+		Compute([&InInputMeshes](int32 Idx) { return InInputMeshes[Idx]; }, InInputMeshes.Num(), MeshBoundsIn);
+	}
+
+	void Compute(TArrayView<const TUniquePtr<TriangleMeshType>> InInputMeshes, TArrayView<const FAxisAlignedBox3d> MeshBoundsIn = TArrayView<const FAxisAlignedBox3d>())
+	{
+		Compute([&InInputMeshes](int32 Idx) { return InInputMeshes[Idx].Get(); }, InInputMeshes.Num(), MeshBoundsIn);
+	}
+
+	void Compute(TFunctionRef<const TriangleMeshType*(int32)> GetInputMesh, int32 NumInputMeshes, TArrayView<const FAxisAlignedBox3d> MeshBoundsIn = TArrayView<const FAxisAlignedBox3d>())
+	{
+		MeshBounds.Reset();
+
+		if (MeshBoundsIn.Num() == NumInputMeshes)
+		{
+			MeshBounds.Append(MeshBoundsIn);
+		}
+		else
 		{
-			InitBounds();
+			InitBounds(GetInputMesh, NumInputMeshes);
 		}
 
 		FAxisAlignedBox3d CombinedBounds;
@@ -122,8 +143,9 @@ public:
 		}
 
 		Nests.Empty();
+		SkippedMeshIndices.Empty();
 
-		int32 N = InputMeshes.Num();
+		int32 N = NumInputMeshes;
 		TArray<int> MeshIndices;
 		MeshIndices.SetNum(N);
 		for (int i = 0; i < N; i++)
@@ -142,7 +164,7 @@ public:
 		double DimScaleFactor = 1.0 / (CombinedBounds.MaxDim() + KINDA_SMALL_NUMBER);
 		for (int i = 0; i < N; i++)
 		{
-			double SignedVolume = TMeshQueries<TriangleMeshType>::GetVolumeNonWatertight(InputMeshes[i], DimScaleFactor);
+			double SignedVolume = TMeshQueries<TriangleMeshType>::GetVolumeNonWatertight(*GetInputMesh(i), DimScaleFactor);
 			Info[i].bPositiveVolume = SignedVolume >= 0;
 			Info[i].Volume = FMathd::Abs(SignedVolume);
 		}
@@ -163,7 +185,7 @@ public:
 		{
 			int PIdx = MeshIndices[ParentCand];
 			const FMeshInfo& PInfo = Info[PIdx];
-			const TriangleMeshType& PMesh = InputMeshes[PIdx];
+			const TriangleMeshType& PMesh = *GetInputMesh(PIdx);
 			if (NestingMethod == ENestingMethod::InLargestParent && Parent[ParentCand] != -1)
 			{
 				continue; // mesh was already contained in a larger mesh
@@ -176,13 +198,13 @@ public:
 			}
 
 			double WindingTestSign = PInfo.bPositiveVolume ? 1 : -1;
-			TMeshAABBTree3<TriangleMeshType> Spatial(&InputMeshes[PIdx]);
+			TMeshAABBTree3<TriangleMeshType> Spatial(GetInputMesh(PIdx));
 			TFastWindingTree<TriangleMeshType> Winding(&Spatial);
 			for (int ChildCand = ParentCand + 1; ChildCand < N; ChildCand++)
 			{
 				int CIdx = MeshIndices[ChildCand];
 				const FMeshInfo& CInfo = Info[CIdx];
-				const TriangleMeshType& CMesh = InputMeshes[CIdx];
+				const TriangleMeshType& CMesh = *GetInputMesh(CIdx);
 				if (CMesh.VertexCount() == 0)
 				{
 					continue;
@@ -277,13 +299,13 @@ public:
 
 private:
 
-	void InitBounds()
+	void InitBounds(TFunctionRef<const TriangleMeshType*(int32)> GetInputMesh, int32 NumInputMeshes)
 	{
 		MeshBounds.Reset();
-		MeshBounds.Reserve(InputMeshes.Num());
-		for (const TriangleMeshType& Mesh : InputMeshes)
+		MeshBounds.Reserve(NumInputMeshes);
+		for (int32 Idx = 0; Idx < NumInputMeshes; ++Idx)
 		{
-			MeshBounds.Add(TMeshQueries<TriangleMeshType>::GetBounds(Mesh));
+			MeshBounds.Add(TMeshQueries<TriangleMeshType>::GetBounds(*GetInputMesh(Idx)));
 		}
 	}
 };
