Add more options to control island splitting on fracture (threshold distance, whether to connect across triangles)

#Synced-CL 52268326

diff --git a/Engine/Plugins/Experimental/PlanarCutPlugin/Source/PlanarCut/Private/GeometryMeshConversion.cpp b/Engine/Plugins/Experimental/PlanarCutPlugin/Source/PlanarCut/Private/GeometryMeshConversion.cpp
--- a/Engine/Plugins/Experimental/PlanarCutPlugin/Source/PlanarCut/Private/GeometryMeshConversion.cpp
+++ b/Engine/Plugins/Experimental/PlanarCutPlugin/Source/PlanarCut/Private/GeometryMeshConversion.cpp
@@ -2725,7 +2725,7 @@ int32 FDynamicMeshCollection::CutWithMultiplePlanes(
 	const TArrayView<const FPlane>& Planes,
 	double Grout,
 	double CollisionSampleSpacing,
-	bool bSplitIslands,
+	const FIslandSplitSettings& IslandSplitSettings,
 	int32 RandomSeed,
 	FGeometryCollection* Collection,
 	FInternalSurfaceMaterials& InternalSurfaceMaterials,
@@ -2787,7 +2787,7 @@ int32 FDynamicMeshCollection::CutWithMultiplePlanes(
 		TArray<TPair<int32, int32>> CellConnectivity;
 		CellConnectivity.Add(TPair<int32, int32>(0, -1));
 
-		return CutWithCellMeshes(InternalSurfaceMaterials, CellConnectivity, GroutCells, bSplitIslands, Collection, bSetDefaultInternalMaterialsFromCollection, CollisionSampleSpacing);
+		return CutWithCellMeshes(InternalSurfaceMaterials, CellConnectivity, GroutCells, IslandSplitSettings, Collection, bSetDefaultInternalMaterialsFromCollection, CollisionSampleSpacing);
 	}
 
 	TArray<TUniquePtr<FMeshData>> ToCut;
@@ -2877,7 +2877,7 @@ int32 FDynamicMeshCollection::CutWithMultiplePlanes(
 				TArray<TUniquePtr<FDynamicMesh3>> SplitMeshes;
 				for (int UnsplitIdx = 0; UnsplitIdx < 2; UnsplitIdx++)
 				{
-					if (bSplitIslands && SplitIslands(ToCut[ResultIndices[UnsplitIdx]]->AugMesh, SplitMeshes))
+					if (IslandSplitSettings.bSplitIslands && SplitIslands(ToCut[ResultIndices[UnsplitIdx]]->AugMesh, SplitMeshes, IslandSplitSettings))
 					{
 						// TODO: Add a SetMesh(FDynamicMesh3&&) method so we can MoveTemp here
 						ToCut[ResultIndices[UnsplitIdx]]->SetMesh(*SplitMeshes[0]);
@@ -2969,7 +2969,7 @@ int32 FDynamicMeshCollection::CutWithMultiplePlanes(
 }
 
 
-int32 FDynamicMeshCollection::SplitAllIslands(FGeometryCollection* Collection, double CollisionSampleSpacing)
+int32 FDynamicMeshCollection::SplitAllIslands(FGeometryCollection* Collection, double CollisionSampleSpacing, const FIslandSplitSettings& IslandSplitSettings)
 {
 	int32 FirstIdx = -1;
 	TArray<int32> GeometryForRemoval;
@@ -2980,7 +2980,7 @@ int32 FDynamicMeshCollection::SplitAllIslands(FGeometryCollection* Collection, d
 		
 		int32 CreatedGeometryIdx = -1;
 		TArray<TUniquePtr<FDynamicMesh3>> Islands;
-		if (SplitIslands(Surface.AugMesh, Islands))
+		if (SplitIslands(Surface.AugMesh, Islands, IslandSplitSettings))
 		{
 			for (int32 i = 0; i < Islands.Num(); i++)
 			{
@@ -3027,7 +3027,7 @@ int32 FDynamicMeshCollection::SplitAllIslands(FGeometryCollection* Collection, d
 
 
 
-int32 FDynamicMeshCollection::CutWithCellMeshes(const FInternalSurfaceMaterials& InternalSurfaceMaterials, const TArray<TPair<int32, int32>>& CellConnectivity, FCellMeshes& CellMeshes, bool bSplitIslands, FGeometryCollection* Collection, bool bSetDefaultInternalMaterialsFromCollection, double CollisionSampleSpacing)
+int32 FDynamicMeshCollection::CutWithCellMeshes(const FInternalSurfaceMaterials& InternalSurfaceMaterials, const TArray<TPair<int32, int32>>& CellConnectivity, FCellMeshes& CellMeshes, const FIslandSplitSettings& IslandSplitSettings, FGeometryCollection* Collection, bool bSetDefaultInternalMaterialsFromCollection, double CollisionSampleSpacing)
 {
 	// TODO: should we do these cuts in parallel, and the appends sequentially below?
 	int32 FirstIdx = -1;
@@ -3108,7 +3108,7 @@ int32 FDynamicMeshCollection::CutWithCellMeshes(const FInternalSurfaceMaterials&
 					}
 					int32 CreatedGeometryIdx = -1;
 					TArray<TUniquePtr<FDynamicMesh3>> Islands;
-					if (bSplitIslands && SplitIslands(AugBoolResult, Islands))
+					if (IslandSplitSettings.bSplitIslands && SplitIslands(AugBoolResult, Islands, IslandSplitSettings))
 					{
 						for (int32 i = 0; i < Islands.Num(); i++)
 						{
@@ -3226,32 +3226,28 @@ bool FDynamicMeshCollection::IsNeighboring(
 }
 
 // Split mesh into connected components, including implicit connections by co-located vertices
-bool FDynamicMeshCollection::SplitIslands(FDynamicMesh3& Source, TArray<TUniquePtr<FDynamicMesh3>>& SeparatedMeshes, double SnapDistance)
+bool FDynamicMeshCollection::SplitIslands(FDynamicMesh3& Source, TArray<TUniquePtr<FDynamicMesh3>>& SeparatedMeshes, const FIslandSplitSettings& Settings)
 {
 	SeparatedMeshes.Reset();
 
-	TPointHashGrid3d<int> VertHash(SnapDistance * 10, -1);
-	FDisjointSet VertComponents(Source.MaxVertexID());
-	// Add Source vertices to hash & disjoint sets
-	for (int VID : Source.VertexIndicesItr())
+	FVertexConnectedComponents VertComponents;
+	VertComponents.Init(Source);
+	VertComponents.ConnectTriangles(Source);
+
+	if (Settings.CloseVertexDistance > 0)
 	{
-		FVector3d Pt = Source.GetVertex(VID);
-		VertHash.EnumeratePointsInBall(Pt, SnapDistance, [&Source, Pt](int OtherVID) {return DistanceSquared(Pt, Source.GetVertex(OtherVID)); }, [&](int32 NbrVID, double DSq) -> bool
-			{
-				VertComponents.UnionSequential(VID, NbrVID);
-				return true;
-			});
-		VertHash.InsertPointUnsafe(VID, Pt);
+		VertComponents.ConnectCloseVertices(Source, Settings.CloseVertexDistance);
 	}
-	for (FIndex3i Tri : Source.TrianglesItr())
+
+	if (Settings.VertexToSurfaceBridgeDistance > 0)
 	{
-		VertComponents.Union(Tri.A, Tri.B);
-		VertComponents.Union(Tri.B, Tri.C);
+		TMeshAABBTree3<FDynamicMesh3> Spatial(&Source);
+		VertComponents.ConnectVerticesToNearestDifferentComponent(Spatial, Settings.VertexToSurfaceBridgeDistance);
 	}
 
 	bool bWasSplit = FDynamicMeshEditor::SplitMesh(&Source, SeparatedMeshes, false, [&Source, &VertComponents](int TID)
 		{
-			return (int)VertComponents.Find(Source.GetTriangle(TID).A);
+			return VertComponents.GetComponent(Source.GetTriangle(TID).A);
 		});
 
 	if (bWasSplit)
diff --git a/Engine/Plugins/Experimental/PlanarCutPlugin/Source/PlanarCut/Private/GeometryMeshConversion.h b/Engine/Plugins/Experimental/PlanarCutPlugin/Source/PlanarCut/Private/GeometryMeshConversion.h
--- a/Engine/Plugins/Experimental/PlanarCutPlugin/Source/PlanarCut/Private/GeometryMeshConversion.h
+++ b/Engine/Plugins/Experimental/PlanarCutPlugin/Source/PlanarCut/Private/GeometryMeshConversion.h
@@ -3,6 +3,7 @@
 
 #include "DynamicMesh/DynamicMesh3.h"
 #include "GeometryCollection/GeometryCollection.h"
+#include "PlanarCut.h"
 #include "Spatial/PointHashGrid3.h"
 
 namespace UE::Geometry { enum class EEdgeRefineFlags; }
@@ -263,7 +264,7 @@ struct FDynamicMeshCollection
 		const TArrayView<const FPlane>& Planes,
 		double Grout,
 		double CollisionSampleSpacing,
-		bool bSplitIslands,
+		const FIslandSplitSettings& IslandSplitSettings,
 		int32 RandomSeed,
 		FGeometryCollection* Collection,
 		FInternalSurfaceMaterials& InternalSurfaceMaterials,
@@ -282,7 +283,7 @@ struct FDynamicMeshCollection
 	 * @param CollisionSampleSpacing If > 0, new geometry will have collision samples added (vertices not on any triangles) to fill any gaps greater than the this size
 	 * @return Index of the first created geometry
 	 */
-	int32 CutWithCellMeshes(const FInternalSurfaceMaterials& InternalSurfaceMaterials, const TArray<TPair<int32, int32>>& CellConnectivity, FCellMeshes& CellMeshes, bool bSplitIslands, FGeometryCollection* Collection, bool bSetDefaultInternalMaterialsFromCollection, double CollisionSampleSpacing);
+	int32 CutWithCellMeshes(const FInternalSurfaceMaterials& InternalSurfaceMaterials, const TArray<TPair<int32, int32>>& CellConnectivity, FCellMeshes& CellMeshes, const FIslandSplitSettings& IslandSplitSettings, FGeometryCollection* Collection, bool bSetDefaultInternalMaterialsFromCollection, double CollisionSampleSpacing);
 
 
 	/**
@@ -292,7 +293,7 @@ struct FDynamicMeshCollection
 	 * @param CollisionSampleSpacing If > 0, new geometry will have collision samples added (vertices not on any triangles) to fill any gaps greater than the this size
 	 * @return Index of the first created geometry, or -1 if nothing was split
 	 */
-	int32 SplitAllIslands(FGeometryCollection* Collection, double CollisionSampleSpacing);
+	int32 SplitAllIslands(FGeometryCollection* Collection, double CollisionSampleSpacing, const FIslandSplitSettings& IslandSplitSettings = FIslandSplitSettings());
 
 	static void SetVisibility(FGeometryCollection& Collection, int32 GeometryIdx, bool bVisible)
 	{
@@ -304,7 +305,7 @@ struct FDynamicMeshCollection
 	}
 
 	// Split mesh into connected components, including implicit connections by co-located vertices
-	bool SplitIslands(UE::Geometry::FDynamicMesh3& Source, TArray<TUniquePtr<UE::Geometry::FDynamicMesh3>>& SeparatedMeshes, double ToleranceDistance = 1e-3);
+	bool SplitIslands(UE::Geometry::FDynamicMesh3& Source, TArray<TUniquePtr<UE::Geometry::FDynamicMesh3>>& SeparatedMeshes, const FIslandSplitSettings& IslandSplitSettings = FIslandSplitSettings());
 
 	FString GetBoneName(FGeometryCollection& Output, int TransformParent, int SubPartIndex)
 	{
diff --git a/Engine/Plugins/Experimental/PlanarCutPlugin/Source/PlanarCut/Private/PlanarCut.cpp b/Engine/Plugins/Experimental/PlanarCutPlugin/Source/PlanarCut/Private/PlanarCut.cpp
--- a/Engine/Plugins/Experimental/PlanarCutPlugin/Source/PlanarCut/Private/PlanarCut.cpp
+++ b/Engine/Plugins/Experimental/PlanarCutPlugin/Source/PlanarCut/Private/PlanarCut.cpp
@@ -813,11 +813,11 @@ int32 CutWithPlanarCells(
 	bool bSetDefaultInternalMaterialsFromCollection,
 	FProgressCancel* Progress,
 	FVector CellsOrigin,
-	bool bSplitIslands
+	FIslandSplitSettings IslandSplitSettings
 )
 {
 	TArray<int32> TransformIndices { TransformIdx };
-	return CutMultipleWithPlanarCells(Cells, Source, TransformIndices, Grout, CollisionSampleSpacing, RandomSeed, TransformCollection, bIncludeOutsideCellInOutput, bSetDefaultInternalMaterialsFromCollection, Progress, CellsOrigin, bSplitIslands);
+	return CutMultipleWithPlanarCells(Cells, Source, TransformIndices, Grout, CollisionSampleSpacing, RandomSeed, TransformCollection, bIncludeOutsideCellInOutput, bSetDefaultInternalMaterialsFromCollection, Progress, CellsOrigin, IslandSplitSettings);
 }
 
 
@@ -832,7 +832,7 @@ int32 CutMultipleWithMultiplePlanes(
 	const TOptional<FTransform>& TransformCollection,
 	bool bSetDefaultInternalMaterialsFromCollection,
 	FProgressCancel* Progress,
-	bool bSplitIslands
+	FIslandSplitSettings IslandSplitSettings
 )
 {
 	FProgressCancel::FProgressScope PrepareScope = FProgressCancel::CreateScopeTo(Progress, .1, LOCTEXT("CutWithMultiplePlanesInit", "Preparing to cut with planes"));
@@ -873,7 +873,7 @@ int32 CutMultipleWithMultiplePlanes(
 	FProgressCancel::FProgressScope CutScope = FProgressCancel::CreateScopeTo(Progress, .99, LOCTEXT("CutWithMultiplePlanesBody", "Cutting with planes"));
 
 	int32 NewGeomStartIdx = -1;
-	NewGeomStartIdx = MeshCollection.CutWithMultiplePlanes(CenteredPlanes, Grout, CollisionSampleSpacing, bSplitIslands, RandomSeed, &Collection, InternalSurfaceMaterials, bSetDefaultInternalMaterialsFromCollection, Progress);
+	NewGeomStartIdx = MeshCollection.CutWithMultiplePlanes(CenteredPlanes, Grout, CollisionSampleSpacing, IslandSplitSettings, RandomSeed, &Collection, InternalSurfaceMaterials, bSetDefaultInternalMaterialsFromCollection, Progress);
 
 	CutScope.Done();
 	if (Progress && Progress->Cancelled())
@@ -947,7 +947,7 @@ int32 CutMultipleWithPlanarCells(
 	bool bSetDefaultInternalMaterialsFromCollection,
 	FProgressCancel* Progress,
 	FVector CellsOrigin,
-	bool bSplitIslands
+	FIslandSplitSettings IslandSplitSettings
 )
 {
 	FProgressCancel::FProgressScope CreateMeshCollectionScope = FProgressCancel::CreateScopeTo(Progress, .1);
@@ -981,7 +981,7 @@ int32 CutMultipleWithPlanarCells(
 
 	FProgressCancel::FProgressScope CutScope = FProgressCancel::CreateScopeTo(Progress, .99);
 	int32 NewGeomStartIdx = -1;
-	NewGeomStartIdx = MeshCollection.CutWithCellMeshes(Cells.InternalSurfaceMaterials, Cells.PlaneCells, CellMeshes, bSplitIslands, &Source, bSetDefaultInternalMaterialsFromCollection, CollisionSampleSpacing);
+	NewGeomStartIdx = MeshCollection.CutWithCellMeshes(Cells.InternalSurfaceMaterials, Cells.PlaneCells, CellMeshes, IslandSplitSettings, &Source, bSetDefaultInternalMaterialsFromCollection, CollisionSampleSpacing);
 	CutScope.Done();
 
 	if (Progress && Progress->Cancelled())
@@ -999,7 +999,8 @@ int32 SplitIslands(
 	FGeometryCollection& Collection,
 	const TArrayView<const int32>& TransformIndices,
 	double CollisionSampleSpacing,
-	FProgressCancel* Progress
+	FProgressCancel* Progress,
+	FIslandSplitSettings IslandSplitSettings
 )
 {
 	FProgressCancel::FProgressScope CreateMeshCollectionScope = FProgressCancel::CreateScopeTo(Progress, .1);
@@ -1016,7 +1017,7 @@ int32 SplitIslands(
 
 	FProgressCancel::FProgressScope SplitScope = FProgressCancel::CreateScopeTo(Progress, .99);
 	int32 NewGeomStartIdx = -1;
-	NewGeomStartIdx = MeshCollection.SplitAllIslands(&Collection, CollisionSampleSpacing);
+	NewGeomStartIdx = MeshCollection.SplitAllIslands(&Collection, CollisionSampleSpacing, IslandSplitSettings);
 	SplitScope.Done();
 
 	if (Progress && Progress->Cancelled())
@@ -1145,7 +1146,7 @@ int32 CutWithMesh(
 	const TOptional<FTransform>& TransformCollection,
 	bool bSetDefaultInternalMaterialsFromCollection,
 	FProgressCancel* Progress,
-	bool bSplitIslands
+	FIslandSplitSettings IslandSplitSettings
 )
 {
 	FProgressCancel::FProgressScope PrepareScope = FProgressCancel::CreateScopeTo(Progress, .1);
@@ -1185,7 +1186,7 @@ int32 CutWithMesh(
 	}
 	FProgressCancel::FProgressScope CutScope = FProgressCancel::CreateScopeTo(Progress, .99);
 
-	NewGeomStartIdx = MeshCollection.CutWithCellMeshes(InternalSurfaceMaterials, CellConnectivity, CellMeshes, bSplitIslands, &Collection, bSetDefaultInternalMaterialsFromCollection, CollisionSampleSpacing);
+	NewGeomStartIdx = MeshCollection.CutWithCellMeshes(InternalSurfaceMaterials, CellConnectivity, CellMeshes, IslandSplitSettings, &Collection, bSetDefaultInternalMaterialsFromCollection, CollisionSampleSpacing);
 
 	CutScope.Done();
 	if (Progress && Progress->Cancelled())
diff --git a/Engine/Plugins/Experimental/PlanarCutPlugin/Source/PlanarCut/Public/PlanarCut.h b/Engine/Plugins/Experimental/PlanarCutPlugin/Source/PlanarCut/Public/PlanarCut.h
--- a/Engine/Plugins/Experimental/PlanarCutPlugin/Source/PlanarCut/Public/PlanarCut.h
+++ b/Engine/Plugins/Experimental/PlanarCutPlugin/Source/PlanarCut/Public/PlanarCut.h
@@ -68,6 +68,25 @@ struct FNoiseSettings
 	}
 };
 
+/** Settings controlling how disconnected islands are detected and split after a cut */
+struct FIslandSplitSettings
+{
+
+	FIslandSplitSettings(bool bInSplitIslands = true, double InCloseVertexDistance = 1e-3, double InVertexToSurfaceBridgeDistance = 0)
+		: bSplitIslands(bInSplitIslands), CloseVertexDistance(InCloseVertexDistance),
+		  VertexToSurfaceBridgeDistance(InVertexToSurfaceBridgeDistance) {}
+
+	/** Whether to detect and split disconnected islands */
+	bool bSplitIslands = true;
+
+	/** Distance threshold for connecting spatially-close vertices before computing islands */
+	double CloseVertexDistance = 1e-3;
+
+	/** If > 0, bridge separate islands whose surfaces are within this vertex-to-triangle distance */
+	double VertexToSurfaceBridgeDistance = 0;
+
+};
+
 // auxiliary structure for FPlanarCells to carry material info
 struct FInternalSurfaceMaterials
 {
@@ -204,7 +223,7 @@ struct FPlanarCells
  * @param bIncludeOutsideCellInOutput	If true, geometry that was not inside any of the cells (e.g. was outside of the bounds of all cutting geometry) will still be included in the output; if false, it will be discarded.
  * @param Progress						Optionally tracks progress and supports early-cancel
  * @param CellsOrigin					Optionally provide a local origin of the cutting Cells
- * @param bSplitIslands					Whether to detect connected components and split them apart after the cut
+ * @param IslandSplitSettings			Controls detection and splitting of disconnected components after the cut
  * @return	index of first new geometry in the Output GeometryCollection, or -1 if no geometry was added
  */
 int32 PLANARCUT_API CutWithPlanarCells(
@@ -219,7 +238,7 @@ int32 PLANARCUT_API CutWithPlanarCells(
 	bool bSetDefaultInternalMaterialsFromCollection = true,
 	FProgressCancel* Progress = nullptr,
 	FVector CellsOrigin = FVector::ZeroVector,
-	bool bSplitIslands = true
+	FIslandSplitSettings IslandSplitSettings = FIslandSplitSettings()
 );
 
 /**
@@ -260,7 +279,7 @@ void PLANARCUT_API CreateCuttingSurfacePreview(
  * @param bIncludeOutsideCellInOutput	If true, geometry that was not inside any of the cells (e.g. was outside of the bounds of all cutting geometry) will still be included in the output; if false, it will be discarded.
  * @param Progress						Optionally tracks progress and supports early-cancel
  * @param CellsOrigin					Optionally provide a local origin of the cutting Cells
- * @param bSplitIslands					Whether to detect connected components and split them apart after the cut
+ * @param IslandSplitSettings			Controls detection and splitting of disconnected components after the cut
  * @return	index of first new geometry in the Output GeometryCollection, or -1 if no geometry was added
  */
 int32 PLANARCUT_API CutMultipleWithPlanarCells(
@@ -275,7 +294,7 @@ int32 PLANARCUT_API CutMultipleWithPlanarCells(
 	bool bSetDefaultInternalMaterialsFromCollection = true,
 	FProgressCancel* Progress = nullptr,
 	FVector CellsOrigin = FVector::ZeroVector,
-	bool bSplitIslands = true
+	FIslandSplitSettings IslandSplitSettings = FIslandSplitSettings()
 );
 
 /**
@@ -291,7 +310,8 @@ int32 PLANARCUT_API SplitIslands(
 	FGeometryCollection& Collection,
 	const TArrayView<const int32>& TransformIndices,
 	double CollisionSampleSpacing,
-	FProgressCancel* Progress = nullptr
+	FProgressCancel* Progress = nullptr,
+	FIslandSplitSettings IslandSplitSettings = FIslandSplitSettings()
 );
 
 /**
@@ -306,7 +326,7 @@ int32 PLANARCUT_API SplitIslands(
  * @param RandomSeed				Seed to be used for random noise displacement
  * @param TransformCollection		Optional transform of the whole geometry collection; if unset, defaults to Identity
  * @param Progress					Optionally tracks progress and supports early-cancel
- * @param bSplitIslands					Whether to detect connected components and split them apart after the cut
+ * @param IslandSplitSettings			Controls detection and splitting of disconnected components after the cut
  * @return	index of first new geometry in the Output GeometryCollection, or -1 if no geometry was added
  */
 int32 PLANARCUT_API CutMultipleWithMultiplePlanes(
@@ -320,7 +340,7 @@ int32 PLANARCUT_API CutMultipleWithMultiplePlanes(
 	const TOptional<FTransform>& TransformCollection = TOptional<FTransform>(),
 	bool bSetDefaultInternalMaterialsFromCollection = true,
 	FProgressCancel* Progress = nullptr,
-	bool bSplitIslands = true
+	FIslandSplitSettings IslandSplitSettings = FIslandSplitSettings()
 );
 
 
@@ -488,7 +508,7 @@ int32 PLANARCUT_API AddCollisionSampleVertices(double TargetSpacing, FGeometryCo
  * @param CollisionSampleSpacing	Target spacing between collision sample vertices
  * @param TransformCollection		Optional transform of the collection; if unset, defaults to Identity
  * @param Progress					Optionally tracks progress and supports early-cancel
- * @param bSplitIslands				Whether to detect connected components and split them apart after the cut
+ * @param IslandSplitSettings			Controls detection and splitting of disconnected components after the cut
  * @return index of first new geometry in the Output GeometryCollection, or -1 if no geometry was added
  */
 int32 PLANARCUT_API CutWithMesh(
@@ -501,7 +521,7 @@ int32 PLANARCUT_API CutWithMesh(
 	const TOptional<FTransform>& TransformCollection = TOptional<FTransform>(),
 	bool bSetDefaultInternalMaterialsFromCollection = true,
 	FProgressCancel* Progress = nullptr,
-	bool bSplitIslands = true
+	FIslandSplitSettings IslandSplitSettings = FIslandSplitSettings()
 );
 
 /// Convert a mesh description to a dynamic mesh *specifically* augmented / designed for cutting geometry collections, to be passed to CutWithMesh
@@ -517,7 +537,7 @@ inline int32 CutWithMesh(
 	const TOptional<FTransform>& TransformCollection = TOptional<FTransform>(),
 	bool bSetDefaultInternalMaterialsFromCollection = true,
 	FProgressCancel* Progress = nullptr,
-	bool bSplitIslands = true
+	FIslandSplitSettings IslandSplitSettings = FIslandSplitSettings()
 )
 {
 	int32 NumUVLayers = Collection.NumUVLayers();
@@ -530,7 +550,7 @@ inline int32 CutWithMesh(
 		TransformCollection,
 		bSetDefaultInternalMaterialsFromCollection,
 		Progress,
-		bSplitIslands
+		IslandSplitSettings
 		);
 }
 
