fix double-vs-int comparison in geometrymeshconversion.cpp

#Synced-CL 52268326

diff --git a/Engine/Plugins/Experimental/PlanarCutPlugin/Source/PlanarCut/Private/GeometryMeshConversion.cpp b/Engine/Plugins/Experimental/PlanarCutPlugin/Source/PlanarCut/Private/GeometryMeshConversion.cpp
--- a/Engine/Plugins/Experimental/PlanarCutPlugin/Source/PlanarCut/Private/GeometryMeshConversion.cpp
+++ b/Engine/Plugins/Experimental/PlanarCutPlugin/Source/PlanarCut/Private/GeometryMeshConversion.cpp
@@ -3234,12 +3234,12 @@ bool FDynamicMeshCollection::SplitIslands(FDynamicMesh3& Source, TArray<TUniqueP
 	VertComponents.Init(Source);
 	VertComponents.ConnectTriangles(Source);
 
-	if (Settings.CloseVertexDistance > 0)
+	if (Settings.CloseVertexDistance > 0.)
 	{
 		VertComponents.ConnectCloseVertices(Source, Settings.CloseVertexDistance);
 	}
 
-	if (Settings.VertexToSurfaceBridgeDistance > 0)
+	if (Settings.VertexToSurfaceBridgeDistance > 0.)
 	{
 		TMeshAABBTree3<FDynamicMesh3> Spatial(&Source);
 		VertComponents.ConnectVerticesToNearestDifferentComponent(Spatial, Settings.VertexToSurfaceBridgeDistance);
